Fix crash when launching MediaBundle in Standalone

diff --git a/Engine/Plugins/Media/MediaFrameworkUtilities/Source/MediaFrameworkUtilities/Private/MediaBundle.cpp b/Engine/Plugins/Media/MediaFrameworkUtilities/Source/MediaFrameworkUtilities/Private/MediaBundle.cpp
--- a/Engine/Plugins/Media/MediaFrameworkUtilities/Source/MediaFrameworkUtilities/Private/MediaBundle.cpp
+++ b/Engine/Plugins/Media/MediaFrameworkUtilities/Source/MediaFrameworkUtilities/Private/MediaBundle.cpp
@@ -75,10 +75,13 @@ void UMediaBundle::CloseMediaSource()
 void UMediaBundle::SetIsValidMaterialParameter(bool bIsValid)
 {
 #if WITH_EDITOR
-	if (UMaterialInstanceConstant* Instance = Cast<UMaterialInstanceConstant>(Material))
+	if (GIsEditor)
 	{
-		Instance->SetScalarParameterValueEditorOnly(FMaterialParameterInfo(MediaBundleMaterialParametersName::IsValidMediaName), bIsValid ? 1.0f : 0.0f);
-		Instance->PostEditChange();
+		if (UMaterialInstanceConstant* Instance = Cast<UMaterialInstanceConstant>(Material))
+		{
+			Instance->SetScalarParameterValueEditorOnly(FMaterialParameterInfo(MediaBundleMaterialParametersName::IsValidMediaName), bIsValid ? 1.0f : 0.0f);
+			Instance->PostEditChange();
+		}
 	}
 #endif
 }
@@ -134,47 +137,50 @@ void UMediaBundle::RefreshLensDisplacementMap()
 void UMediaBundle::CreateInternalsEditor()
 {
 #if WITH_EDITOR && WITH_EDITORONLY_DATA
-	IAssetTools& AssetTools = FModuleManager::Get().LoadModuleChecked<FAssetToolsModule>("AssetTools").Get();
-
-	const FString ParentName = GetOuter()->GetName() + "_InnerAssets/";
-	FString OutAssetName;
-	FString OutPackageName;
-
-	//Create MediaPlayer
-	AssetTools.CreateUniqueAssetName(*(ParentName + TEXT("/MediaP_") + GetName()), TEXT(""), OutPackageName, OutAssetName);
-	MediaPlayer = Cast<UMediaPlayer>(AssetTools.CreateAsset(OutAssetName, ParentName, UMediaPlayer::StaticClass(), nullptr));
-	MediaPlayer->AffectedByPIEHandling = false;
-
-	//Create MediaTexture 
-	AssetTools.CreateUniqueAssetName(*(ParentName + TEXT("/T_") + GetName() + TEXT("_BC")), TEXT(""), OutPackageName, OutAssetName);
-	MediaTexture = Cast<UMediaTexture>(AssetTools.CreateAsset(OutAssetName, ParentName, UMediaTexture::StaticClass(), nullptr));
-	MediaTexture->SetDefaultMediaPlayer(MediaPlayer);
-	MediaTexture->SetMediaPlayer(MediaPlayer);
-	MediaTexture->UpdateResource();
-
-	//Create LensDisplacementMap RenderTarget
-	AssetTools.CreateUniqueAssetName(*(ParentName + TEXT("/RT_") + GetName() + TEXT("_LensDisplacement")), TEXT(""), OutPackageName, OutAssetName);
-	LensDisplacementMap = Cast<UTextureRenderTarget2D>(AssetTools.CreateAsset(OutAssetName, ParentName, UTextureRenderTarget2D::StaticClass(), nullptr));
-	LensDisplacementMap->RenderTargetFormat = RTF_RGBA16f;
-	LensDisplacementMap->InitAutoFormat(256, 256);
-	LensDisplacementMap->UpdateResource();
-
-	//Create MaterialInstanceConstant
-	UMaterialInstanceConstantFactoryNew* Factory = NewObject<UMaterialInstanceConstantFactoryNew>();
-	Factory->InitialParent = DefaultMaterial;
-
-	AssetTools.CreateUniqueAssetName(*(ParentName + TEXT("/MI_") + GetName()), TEXT(""), OutPackageName, OutAssetName);
-	UMaterialInstanceConstant* NewMaterial = Cast<UMaterialInstanceConstant>(AssetTools.CreateAsset(OutAssetName, ParentName, UMaterialInstanceConstant::StaticClass(), Factory));
-	NewMaterial->SetTextureParameterValueEditorOnly(FMaterialParameterInfo(MediaBundleMaterialParametersName::MediaTextureName), MediaTexture);
-	NewMaterial->SetTextureParameterValueEditorOnly(FMaterialParameterInfo(MediaBundleMaterialParametersName::FailedTextureName), DefaultFailedTexture);
-	NewMaterial->SetTextureParameterValueEditorOnly(FMaterialParameterInfo(MediaBundleMaterialParametersName::LensDisplacementMapTextureName), LensDisplacementMap);
-	NewMaterial->PostEditChange();
-	Material = NewMaterial;
-
-	//If we are creating a new object, set the default actor class. Duplicates won't change this.
-	if (MediaBundleActorClass == nullptr)
+	if (GIsEditor)
 	{
-		MediaBundleActorClass = DefaultActorClass;
+		IAssetTools& AssetTools = FModuleManager::Get().LoadModuleChecked<FAssetToolsModule>("AssetTools").Get();
+
+		const FString ParentName = GetOuter()->GetName() + "_InnerAssets/";
+		FString OutAssetName;
+		FString OutPackageName;
+
+		//Create MediaPlayer
+		AssetTools.CreateUniqueAssetName(*(ParentName + TEXT("/MediaP_") + GetName()), TEXT(""), OutPackageName, OutAssetName);
+		MediaPlayer = Cast<UMediaPlayer>(AssetTools.CreateAsset(OutAssetName, ParentName, UMediaPlayer::StaticClass(), nullptr));
+		MediaPlayer->AffectedByPIEHandling = false;
+
+		//Create MediaTexture 
+		AssetTools.CreateUniqueAssetName(*(ParentName + TEXT("/T_") + GetName() + TEXT("_BC")), TEXT(""), OutPackageName, OutAssetName);
+		MediaTexture = Cast<UMediaTexture>(AssetTools.CreateAsset(OutAssetName, ParentName, UMediaTexture::StaticClass(), nullptr));
+		MediaTexture->SetDefaultMediaPlayer(MediaPlayer);
+		MediaTexture->SetMediaPlayer(MediaPlayer);
+		MediaTexture->UpdateResource();
+
+		//Create LensDisplacementMap RenderTarget
+		AssetTools.CreateUniqueAssetName(*(ParentName + TEXT("/RT_") + GetName() + TEXT("_LensDisplacement")), TEXT(""), OutPackageName, OutAssetName);
+		LensDisplacementMap = Cast<UTextureRenderTarget2D>(AssetTools.CreateAsset(OutAssetName, ParentName, UTextureRenderTarget2D::StaticClass(), nullptr));
+		LensDisplacementMap->RenderTargetFormat = RTF_RGBA16f;
+		LensDisplacementMap->InitAutoFormat(256, 256);
+		LensDisplacementMap->UpdateResource();
+
+		//Create MaterialInstanceConstant
+		UMaterialInstanceConstantFactoryNew* Factory = NewObject<UMaterialInstanceConstantFactoryNew>();
+		Factory->InitialParent = DefaultMaterial;
+
+		AssetTools.CreateUniqueAssetName(*(ParentName + TEXT("/MI_") + GetName()), TEXT(""), OutPackageName, OutAssetName);
+		UMaterialInstanceConstant* NewMaterial = Cast<UMaterialInstanceConstant>(AssetTools.CreateAsset(OutAssetName, ParentName, UMaterialInstanceConstant::StaticClass(), Factory));
+		NewMaterial->SetTextureParameterValueEditorOnly(FMaterialParameterInfo(MediaBundleMaterialParametersName::MediaTextureName), MediaTexture);
+		NewMaterial->SetTextureParameterValueEditorOnly(FMaterialParameterInfo(MediaBundleMaterialParametersName::FailedTextureName), DefaultFailedTexture);
+		NewMaterial->SetTextureParameterValueEditorOnly(FMaterialParameterInfo(MediaBundleMaterialParametersName::LensDisplacementMapTextureName), LensDisplacementMap);
+		NewMaterial->PostEditChange();
+		Material = NewMaterial;
+
+		//If we are creating a new object, set the default actor class. Duplicates won't change this.
+		if (MediaBundleActorClass == nullptr)
+		{
+			MediaBundleActorClass = DefaultActorClass;
+		}
 	}
 #endif // WITH_EDITOR && WITH_EDITORONLY_DATA
 }
