Copying //UE4/Release-Staging-4.20 to //UE4/Dev-Main (Source: //UE4/Release-4.20 @ 4212847)

============================
  MAJOR FEATURES & CHANGES
============================

Change 4212847 by [REDACTED_PERSON]

	Fix foliage serialization issue introduced in CL 4206551

Change 4208376 by [REDACTED_PERSON]

	Temporarily removing a block dealing with the VR shared depth buffers that was saving us on a theoretically uneeded depth copy. For the moment, the added depth copy ensures that the depth buffer is set as a input param when needed, and that the swapchain behavior that we do under the hood doesn't interfere with the state caching system. We need to better understand what's going wrong here before we can add this block back in.

Change 4208297 by [REDACTED_PERSON]

	Updating lighting on examples 3.1 and 3.3 of the Niagara level. Fixing some spacing issues in the text of exmaple 1.5.

Change 4207996 by [REDACTED_PERSON]

	Checking in new icons and launch images.

Change 4207953 by [REDACTED_PERSON]

	Editing QAGame test content for VolumeMaterials. Added empty material for particle lights

Change 4207860 by [REDACTED_PERSON]

	Fix any key not triggering actions

Change 4207836 by [REDACTED_PERSON]

	New icon from Mitchell.

Change 4207228 by [REDACTED_PERSON]

	Adding all missing iOS icons and launch images.

Change 4207067 by [REDACTED_PERSON]

	Fixed a crash introduced with CL 4183733 - TSharedPtr doesn't work well with ObjC blocks

Change 4206551 by [REDACTED_PERSON]

	Cook foliage for HTML5 using full precision instance transforms
	Remove previous shader fix for this issue

Change 4206512 by jonathan.lindquist

	Submitting missing files
	Cleared with Nick Bullard-Bradley

Change 4206421 by [REDACTED_PERSON]

	Delete any existing temporary workspace created for SyncDepotPath before trying to sync from it.

Change 4205345 by [REDACTED_PERSON]

Change 4204854 by [REDACTED_PERSON]

Change 4204665 by [REDACTED_PERSON]

	Fixed changing media tracks causing video to restart from beginning on Windows.

Change 4204584 by [REDACTED_PERSON]

	Change the order of calls in FMacWindow::Show() to avoid deadlocks when showing certain modal dialog boxes.

Change 4204541 by [REDACTED_PERSON]

	Adding TPS licenses.

Change 4203023 by [REDACTED_PERSON]

	Updating nearly equal tolerance to resolve a scaling regression on some devices.

Change 4202956 by [REDACTED_PERSON]

	Disable Aftermath until a fix for a memory leak is found

Change 4202953 by [REDACTED_PERSON]

	Fixed crash in DC components when the DC is not initialized

Change 4202952 by [REDACTED_PERSON]

	Fix crash when launching MediaBundle in Standalone

Change 4202900 by [REDACTED_PERSON]

	Setting bGenerateOverlapEvents back to private

Change 4202601 by [REDACTED_PERSON]

	Reduce log spam when building PhysX for Mac/IOS/TVOS. Console output causes massive regression in build times under Mono (10m -> 3h).

Change 4202477 by [REDACTED_PERSON]

	UE-61270 RemoteSessionApp does not change orientation connecting to a project in portrait mode on Android

	#Android
	#4.20

	Setting Orientation: FullSensor for Android builds

Change 4201671 by [REDACTED_PERSON]

	Only attempting to load dlls if the MLSDK is set up (avoiding needless warnings for users who aren't using the MLSDK).

Change 4200166 by [REDACTED_PERSON]

	Prevent invalidation panel crash when pushing the current clip state when there isnt one

diff --git a/Engine/Config/BaseEngine.ini b/Engine/Config/BaseEngine.ini
--- a/Engine/Config/BaseEngine.ini
+++ b/Engine/Config/BaseEngine.ini
@@ -2028,7 +2028,7 @@ bEnableLiveLinkForFaceTracking=true
 LiveLinkPublishingPort=11111
 
 [/Script/Engine.RendererSettings]
-r.GPUCrashDebugging=True
+r.GPUCrashDebugging=false
 
 [Messaging]
 bAllowDelayedMessaging=false
\ No newline at end of file
diff --git a/Engine/Plugins/Lumin/MagicLeap/Source/MLSDK/Public/MagicLeapPluginUtil.h b/Engine/Plugins/Lumin/MagicLeap/Source/MLSDK/Public/MagicLeapPluginUtil.h
--- a/Engine/Plugins/Lumin/MagicLeap/Source/MLSDK/Public/MagicLeapPluginUtil.h
+++ b/Engine/Plugins/Lumin/MagicLeap/Source/MLSDK/Public/MagicLeapPluginUtil.h
@@ -69,22 +69,34 @@ public:
 			{
 				DllSearchPaths.Add(VDZILibraryPath);
 			}
-			// The default VDZI dir.
-			DllSearchPaths.Add(FPaths::Combine(*MLSDK, TEXT("VirtualDevice"), TEXT("lib")));
+
+			if (!MLSDK.IsEmpty())
+			{
+				// The default VDZI dir.
+				DllSearchPaths.Add(FPaths::Combine(*MLSDK, TEXT("VirtualDevice"), TEXT("lib")));
+			}
 		}
 #endif
 
 		// The MLSDK DLLs are platform specific and are segregated in directories for each platform.
+
+#if PLATFORM_LUMIN
+		// Lumin uses the system path as we are in device.
+		DllSearchPaths.Add(TEXT("/system/lib64"));
+#else
+
+		if (!MLSDK.IsEmpty())
+		{
 #if PLATFORM_WINDOWS
-		DllSearchPaths.Add(FPaths::Combine(*MLSDK, TEXT("lib"), TEXT("win64")));
+			DllSearchPaths.Add(FPaths::Combine(*MLSDK, TEXT("lib"), TEXT("win64")));
 #elif PLATFORM_LINUX
-		DllSearchPaths.Add(FPaths::Combine(*MLSDK, TEXT("lib"), TEXT("linux64")));
+			DllSearchPaths.Add(FPaths::Combine(*MLSDK, TEXT("lib"), TEXT("linux64")));
 #elif PLATFORM_MAC
-		DllSearchPaths.Add(FPaths::Combine(*MLSDK, TEXT("lib"), TEXT("osx")));
-#elif PLATFORM_LUMIN
-		// Lumin uses the system path as we are in device.
-		DllSearchPaths.Add(TEXT("/system/lib64"));
+			DllSearchPaths.Add(FPaths::Combine(*MLSDK, TEXT("lib"), TEXT("osx")));
 #endif // PLATFORM_WINDOWS
+		}
+
+#endif // PLATFORM_LUMIN
 
 		// Add the search paths to where we will load the DLLs from. For all just add to the UE4
 		// directory listing. But for Windows we also need to manipulate the PATH for the load
diff --git a/Engine/Plugins/Lumin/MagicLeap/Source/MagicLeap/Private/FoundObjectComponent.cpp b/Engine/Plugins/Lumin/MagicLeap/Source/MagicLeap/Private/FoundObjectComponent.cpp
--- a/Engine/Plugins/Lumin/MagicLeap/Source/MagicLeap/Private/FoundObjectComponent.cpp
+++ b/Engine/Plugins/Lumin/MagicLeap/Source/MagicLeap/Private/FoundObjectComponent.cpp
@@ -141,7 +141,7 @@ UFoundObjectComponent::UFoundObjectComponent()
 	SearchVolume->SetCanEverAffectNavigation(false);
 	SearchVolume->CanCharacterStepUpOn = ECanBeCharacterBase::ECB_No;
 	SearchVolume->SetCollisionObjectType(ECollisionChannel::ECC_WorldDynamic);
-	SearchVolume->bGenerateOverlapEvents = false;
+	SearchVolume->SetGenerateOverlapEvents(false);
 
 #if WITH_EDITOR
 	if (GIsEditor)
diff --git a/Engine/Plugins/Lumin/MagicLeap/Source/MagicLeap/Private/MeshTrackerComponent.cpp b/Engine/Plugins/Lumin/MagicLeap/Source/MagicLeap/Private/MeshTrackerComponent.cpp
--- a/Engine/Plugins/Lumin/MagicLeap/Source/MagicLeap/Private/MeshTrackerComponent.cpp
+++ b/Engine/Plugins/Lumin/MagicLeap/Source/MagicLeap/Private/MeshTrackerComponent.cpp
@@ -367,7 +367,7 @@ UMeshTrackerComponent::UMeshTrackerComponent(const FObjectInitializer& ObjectIni
 	BoundingVolume->SetCanEverAffectNavigation(false);
 	BoundingVolume->CanCharacterStepUpOn = ECanBeCharacterBase::ECB_No;
 	BoundingVolume->SetCollisionObjectType(ECollisionChannel::ECC_WorldDynamic);
-	BoundingVolume->bGenerateOverlapEvents = false;
+	BoundingVolume->SetGenerateOverlapEvents(false);
 	// Recommended default box extents for meshing - 10m (5m radius)
 	BoundingVolume->SetBoxExtent(FVector(1000, 1000, 1000), false);
 
diff --git a/Engine/Plugins/Lumin/MagicLeap/Source/MagicLeap/Private/PlanesComponent.cpp b/Engine/Plugins/Lumin/MagicLeap/Source/MagicLeap/Private/PlanesComponent.cpp
--- a/Engine/Plugins/Lumin/MagicLeap/Source/MagicLeap/Private/PlanesComponent.cpp
+++ b/Engine/Plugins/Lumin/MagicLeap/Source/MagicLeap/Private/PlanesComponent.cpp
@@ -176,7 +176,7 @@ UPlanesComponent::UPlanesComponent()
 	SearchVolume->SetCanEverAffectNavigation(false);
 	SearchVolume->CanCharacterStepUpOn = ECanBeCharacterBase::ECB_No;
 	SearchVolume->SetCollisionObjectType(ECollisionChannel::ECC_WorldDynamic);
-	SearchVolume->bGenerateOverlapEvents = false;
+	SearchVolume->SetGenerateOverlapEvents(false);
 	// Recommended default box extents for meshing - 10m (5m radius)
 	SearchVolume->SetBoxExtent(FVector(1000, 1000, 1000), false);
 
diff --git a/Engine/Plugins/Media/MediaFrameworkUtilities/Source/MediaFrameworkUtilities/Private/MediaBundle.cpp b/Engine/Plugins/Media/MediaFrameworkUtilities/Source/MediaFrameworkUtilities/Private/MediaBundle.cpp
--- a/Engine/Plugins/Media/MediaFrameworkUtilities/Source/MediaFrameworkUtilities/Private/MediaBundle.cpp
+++ b/Engine/Plugins/Media/MediaFrameworkUtilities/Source/MediaFrameworkUtilities/Private/MediaBundle.cpp
@@ -75,10 +75,13 @@ void UMediaBundle::CloseMediaSource()
 void UMediaBundle::SetIsValidMaterialParameter(bool bIsValid)
 {
 #if WITH_EDITOR
-	if (UMaterialInstanceConstant* Instance = Cast<UMaterialInstanceConstant>(Material))
+	if (GIsEditor)
 	{
-		Instance->SetScalarParameterValueEditorOnly(FMaterialParameterInfo(MediaBundleMaterialParametersName::IsValidMediaName), bIsValid ? 1.0f : 0.0f);
-		Instance->PostEditChange();
+		if (UMaterialInstanceConstant* Instance = Cast<UMaterialInstanceConstant>(Material))
+		{
+			Instance->SetScalarParameterValueEditorOnly(FMaterialParameterInfo(MediaBundleMaterialParametersName::IsValidMediaName), bIsValid ? 1.0f : 0.0f);
+			Instance->PostEditChange();
+		}
 	}
 #endif
 }
@@ -134,47 +137,50 @@ void UMediaBundle::RefreshLensDisplacementMap()
 void UMediaBundle::CreateInternalsEditor()
 {
 #if WITH_EDITOR && WITH_EDITORONLY_DATA
-	IAssetTools& AssetTools = FModuleManager::Get().LoadModuleChecked<FAssetToolsModule>("AssetTools").Get();
-
-	const FString ParentName = GetOuter()->GetName() + "_InnerAssets/";
-	FString OutAssetName;
-	FString OutPackageName;
-
-	//Create MediaPlayer
-	AssetTools.CreateUniqueAssetName(*(ParentName + TEXT("/MediaP_") + GetName()), TEXT(""), OutPackageName, OutAssetName);
-	MediaPlayer = Cast<UMediaPlayer>(AssetTools.CreateAsset(OutAssetName, ParentName, UMediaPlayer::StaticClass(), nullptr));
-	MediaPlayer->AffectedByPIEHandling = false;
-
-	//Create MediaTexture 
-	AssetTools.CreateUniqueAssetName(*(ParentName + TEXT("/T_") + GetName() + TEXT("_BC")), TEXT(""), OutPackageName, OutAssetName);
-	MediaTexture = Cast<UMediaTexture>(AssetTools.CreateAsset(OutAssetName, ParentName, UMediaTexture::StaticClass(), nullptr));
-	MediaTexture->SetDefaultMediaPlayer(MediaPlayer);
-	MediaTexture->SetMediaPlayer(MediaPlayer);
-	MediaTexture->UpdateResource();
-
-	//Create LensDisplacementMap RenderTarget
-	AssetTools.CreateUniqueAssetName(*(ParentName + TEXT("/RT_") + GetName() + TEXT("_LensDisplacement")), TEXT(""), OutPackageName, OutAssetName);
-	LensDisplacementMap = Cast<UTextureRenderTarget2D>(AssetTools.CreateAsset(OutAssetName, ParentName, UTextureRenderTarget2D::StaticClass(), nullptr));
-	LensDisplacementMap->RenderTargetFormat = RTF_RGBA16f;
-	LensDisplacementMap->InitAutoFormat(256, 256);
-	LensDisplacementMap->UpdateResource();
-
-	//Create MaterialInstanceConstant
-	UMaterialInstanceConstantFactoryNew* Factory = NewObject<UMaterialInstanceConstantFactoryNew>();
-	Factory->InitialParent = DefaultMaterial;
-
-	AssetTools.CreateUniqueAssetName(*(ParentName + TEXT("/MI_") + GetName()), TEXT(""), OutPackageName, OutAssetName);
-	UMaterialInstanceConstant* NewMaterial = Cast<UMaterialInstanceConstant>(AssetTools.CreateAsset(OutAssetName, ParentName, UMaterialInstanceConstant::StaticClass(), Factory));
-	NewMaterial->SetTextureParameterValueEditorOnly(FMaterialParameterInfo(MediaBundleMaterialParametersName::MediaTextureName), MediaTexture);
-	NewMaterial->SetTextureParameterValueEditorOnly(FMaterialParameterInfo(MediaBundleMaterialParametersName::FailedTextureName), DefaultFailedTexture);
-	NewMaterial->SetTextureParameterValueEditorOnly(FMaterialParameterInfo(MediaBundleMaterialParametersName::LensDisplacementMapTextureName), LensDisplacementMap);
-	NewMaterial->PostEditChange();
-	Material = NewMaterial;
-
-	//If we are creating a new object, set the default actor class. Duplicates won't change this.
-	if (MediaBundleActorClass == nullptr)
+	if (GIsEditor)
 	{
-		MediaBundleActorClass = DefaultActorClass;
+		IAssetTools& AssetTools = FModuleManager::Get().LoadModuleChecked<FAssetToolsModule>("AssetTools").Get();
+
+		const FString ParentName = GetOuter()->GetName() + "_InnerAssets/";
+		FString OutAssetName;
+		FString OutPackageName;
+
+		//Create MediaPlayer
+		AssetTools.CreateUniqueAssetName(*(ParentName + TEXT("/MediaP_") + GetName()), TEXT(""), OutPackageName, OutAssetName);
+		MediaPlayer = Cast<UMediaPlayer>(AssetTools.CreateAsset(OutAssetName, ParentName, UMediaPlayer::StaticClass(), nullptr));
+		MediaPlayer->AffectedByPIEHandling = false;
+
+		//Create MediaTexture 
+		AssetTools.CreateUniqueAssetName(*(ParentName + TEXT("/T_") + GetName() + TEXT("_BC")), TEXT(""), OutPackageName, OutAssetName);
+		MediaTexture = Cast<UMediaTexture>(AssetTools.CreateAsset(OutAssetName, ParentName, UMediaTexture::StaticClass(), nullptr));
+		MediaTexture->SetDefaultMediaPlayer(MediaPlayer);
+		MediaTexture->SetMediaPlayer(MediaPlayer);
+		MediaTexture->UpdateResource();
+
+		//Create LensDisplacementMap RenderTarget
+		AssetTools.CreateUniqueAssetName(*(ParentName + TEXT("/RT_") + GetName() + TEXT("_LensDisplacement")), TEXT(""), OutPackageName, OutAssetName);
+		LensDisplacementMap = Cast<UTextureRenderTarget2D>(AssetTools.CreateAsset(OutAssetName, ParentName, UTextureRenderTarget2D::StaticClass(), nullptr));
+		LensDisplacementMap->RenderTargetFormat = RTF_RGBA16f;
+		LensDisplacementMap->InitAutoFormat(256, 256);
+		LensDisplacementMap->UpdateResource();
+
+		//Create MaterialInstanceConstant
+		UMaterialInstanceConstantFactoryNew* Factory = NewObject<UMaterialInstanceConstantFactoryNew>();
+		Factory->InitialParent = DefaultMaterial;
+
+		AssetTools.CreateUniqueAssetName(*(ParentName + TEXT("/MI_") + GetName()), TEXT(""), OutPackageName, OutAssetName);
+		UMaterialInstanceConstant* NewMaterial = Cast<UMaterialInstanceConstant>(AssetTools.CreateAsset(OutAssetName, ParentName, UMaterialInstanceConstant::StaticClass(), Factory));
+		NewMaterial->SetTextureParameterValueEditorOnly(FMaterialParameterInfo(MediaBundleMaterialParametersName::MediaTextureName), MediaTexture);
+		NewMaterial->SetTextureParameterValueEditorOnly(FMaterialParameterInfo(MediaBundleMaterialParametersName::FailedTextureName), DefaultFailedTexture);
+		NewMaterial->SetTextureParameterValueEditorOnly(FMaterialParameterInfo(MediaBundleMaterialParametersName::LensDisplacementMapTextureName), LensDisplacementMap);
+		NewMaterial->PostEditChange();
+		Material = NewMaterial;
+
+		//If we are creating a new object, set the default actor class. Duplicates won't change this.
+		if (MediaBundleActorClass == nullptr)
+		{
+			MediaBundleActorClass = DefaultActorClass;
+		}
 	}
 #endif // WITH_EDITOR && WITH_EDITORONLY_DATA
 }
diff --git a/Engine/Plugins/Media/WmfMedia/Source/WmfMedia/Private/Wmf/WmfMediaSession.cpp b/Engine/Plugins/Media/WmfMedia/Source/WmfMedia/Private/Wmf/WmfMediaSession.cpp
--- a/Engine/Plugins/Media/WmfMedia/Source/WmfMedia/Private/Wmf/WmfMediaSession.cpp
+++ b/Engine/Plugins/Media/WmfMedia/Source/WmfMedia/Private/Wmf/WmfMediaSession.cpp
@@ -930,6 +930,10 @@ bool FWmfMediaSession::CommitTopology(IMFTopology* Topology)
 			UE_LOG(LogWmfMedia, Verbose, TEXT("Session %p: Requesting restart after pending stop"), this);
 
 			RequestedTime = LastTime;
+
+			// Zero LastTime so it matches what WMF is actually doing (rewinding to the beginning)
+			// The PresentationClock has now stopped and FWmfMediaSession::GetTime() will just return LastTime
+			LastTime = FTimespan::Zero();
 		}
 
 		UE_LOG(LogWmfMedia, Verbose, TEXT("Session %p: Requesting topology change after pending stop"), this);
diff --git a/Engine/Plugins/Runtime/nDisplay/Source/DisplayCluster/Private/Blueprints/DisplayClusterBlueprintAPIImpl.cpp b/Engine/Plugins/Runtime/nDisplay/Source/DisplayCluster/Private/Blueprints/DisplayClusterBlueprintAPIImpl.cpp
--- a/Engine/Plugins/Runtime/nDisplay/Source/DisplayCluster/Private/Blueprints/DisplayClusterBlueprintAPIImpl.cpp
+++ b/Engine/Plugins/Runtime/nDisplay/Source/DisplayCluster/Private/Blueprints/DisplayClusterBlueprintAPIImpl.cpp
@@ -13,6 +13,20 @@
 
 
 
+//////////////////////////////////////////////////////////////////////////////////////////////
+// DisplayCluster module API
+//////////////////////////////////////////////////////////////////////////////////////////////
+/** Return if the module has been initialized. */
+bool UDisplayClusterBlueprintAPIImpl::IsModuleInitialized()
+{
+	return IDisplayCluster::Get().IsModuleInitialized();
+}
+
+EDisplayClusterOperationMode UDisplayClusterBlueprintAPIImpl::GetOperationMode()
+{
+	return IDisplayCluster::Get().GetOperationMode();
+}
+
 //////////////////////////////////////////////////////////////////////////////////////////////
 // Cluster API
 //////////////////////////////////////////////////////////////////////////////////////////////
diff --git a/Engine/Plugins/Runtime/nDisplay/Source/DisplayCluster/Private/Blueprints/DisplayClusterBlueprintAPIImpl.h b/Engine/Plugins/Runtime/nDisplay/Source/DisplayCluster/Private/Blueprints/DisplayClusterBlueprintAPIImpl.h
--- a/Engine/Plugins/Runtime/nDisplay/Source/DisplayCluster/Private/Blueprints/DisplayClusterBlueprintAPIImpl.h
+++ b/Engine/Plugins/Runtime/nDisplay/Source/DisplayCluster/Private/Blueprints/DisplayClusterBlueprintAPIImpl.h
@@ -17,6 +17,18 @@ class DISPLAYCLUSTER_API UDisplayClusterBlueprintAPIImpl
 {
 	GENERATED_BODY()
 
+public:
+	//////////////////////////////////////////////////////////////////////////////////////////////
+	// DisplayCluster module API
+	//////////////////////////////////////////////////////////////////////////////////////////////
+	/** Return if the module has been initialized. */
+	UFUNCTION(BlueprintCallable, meta = (DisplayName = "Is module initialized"), Category = "DisplayCluster")
+	virtual bool IsModuleInitialized() override;
+
+	/** Return current operation mode. */
+	UFUNCTION(BlueprintCallable, meta = (DisplayName = "Get operation mode"), Category = "DisplayCluster")
+	virtual EDisplayClusterOperationMode GetOperationMode() override;
+
 public:
 	//////////////////////////////////////////////////////////////////////////////////////////////
 	// Cluster API
diff --git a/Engine/Plugins/Runtime/nDisplay/Source/DisplayCluster/Private/Cluster/Controller/DisplayClusterClusterNodeCtrlBase.cpp b/Engine/Plugins/Runtime/nDisplay/Source/DisplayCluster/Private/Cluster/Controller/DisplayClusterClusterNodeCtrlBase.cpp
--- a/Engine/Plugins/Runtime/nDisplay/Source/DisplayCluster/Private/Cluster/Controller/DisplayClusterClusterNodeCtrlBase.cpp
+++ b/Engine/Plugins/Runtime/nDisplay/Source/DisplayCluster/Private/Cluster/Controller/DisplayClusterClusterNodeCtrlBase.cpp
@@ -34,6 +34,11 @@ FDisplayClusterClusterNodeCtrlBase::~FDisplayClusterClusterNodeCtrlBase()
 //////////////////////////////////////////////////////////////////////////////////////////////
 bool FDisplayClusterClusterNodeCtrlBase::InitializeStereo()
 {
+	if (GDisplayCluster->GetOperationMode() == EDisplayClusterOperationMode::Disabled)
+	{
+		return false;
+	}
+
 	FDisplayClusterConfigViewport ViewportCfg;
 	if (!GDisplayCluster->GetPrivateConfigMgr()->GetLocalViewport(ViewportCfg))
 	{
diff --git a/Engine/Plugins/Runtime/nDisplay/Source/DisplayCluster/Private/Cluster/Controller/DisplayClusterClusterNodeCtrlMaster.cpp b/Engine/Plugins/Runtime/nDisplay/Source/DisplayCluster/Private/Cluster/Controller/DisplayClusterClusterNodeCtrlMaster.cpp
--- a/Engine/Plugins/Runtime/nDisplay/Source/DisplayCluster/Private/Cluster/Controller/DisplayClusterClusterNodeCtrlMaster.cpp
+++ b/Engine/Plugins/Runtime/nDisplay/Source/DisplayCluster/Private/Cluster/Controller/DisplayClusterClusterNodeCtrlMaster.cpp
@@ -48,6 +48,11 @@ void FDisplayClusterClusterNodeCtrlMaster::GetInputData(FDisplayClusterMessage::
 //////////////////////////////////////////////////////////////////////////////////////////////
 bool FDisplayClusterClusterNodeCtrlMaster::InitializeServers()
 {
+	if (GDisplayCluster->GetOperationMode() == EDisplayClusterOperationMode::Disabled)
+	{
+		return false;
+	}
+
 	if (!FDisplayClusterClusterNodeCtrlSlave::InitializeServers())
 	{
 		return false;
diff --git a/Engine/Plugins/Runtime/nDisplay/Source/DisplayCluster/Private/Config/DisplayClusterConfigManager.cpp b/Engine/Plugins/Runtime/nDisplay/Source/DisplayCluster/Private/Config/DisplayClusterConfigManager.cpp
--- a/Engine/Plugins/Runtime/nDisplay/Source/DisplayCluster/Private/Config/DisplayClusterConfigManager.cpp
+++ b/Engine/Plugins/Runtime/nDisplay/Source/DisplayCluster/Private/Config/DisplayClusterConfigManager.cpp
@@ -107,6 +107,11 @@ bool FDisplayClusterConfigManager::GetMasterClusterNode(FDisplayClusterConfigClu
 
 bool FDisplayClusterConfigManager::GetLocalClusterNode(FDisplayClusterConfigClusterNode& node) const
 {
+	if (GDisplayCluster->GetOperationMode() == EDisplayClusterOperationMode::Disabled)
+	{
+		return false;
+	}
+
 	const FString nodeId = GDisplayCluster->GetPrivateClusterMgr()->GetNodeId();
 	return GetItem(CfgClusterNodes, nodeId, node, FString("GetLocalNode"));
 }
diff --git a/Engine/Plugins/Runtime/nDisplay/Source/DisplayCluster/Private/DisplayClusterModule.cpp b/Engine/Plugins/Runtime/nDisplay/Source/DisplayCluster/Private/DisplayClusterModule.cpp
--- a/Engine/Plugins/Runtime/nDisplay/Source/DisplayCluster/Private/DisplayClusterModule.cpp
+++ b/Engine/Plugins/Runtime/nDisplay/Source/DisplayCluster/Private/DisplayClusterModule.cpp
@@ -98,6 +98,9 @@ bool FDisplayClusterModule::Init(EDisplayClusterOperationMode OperationMode)
 		UE_LOG(LogDisplayClusterModule, Error, TEXT("An error occurred during internal initialization"));
 	}
 
+	// Set internal initialization flag
+	bIsModuleInitialized = result;
+
 	return result;
 }
 
diff --git a/Engine/Plugins/Runtime/nDisplay/Source/DisplayCluster/Private/DisplayClusterModule.h b/Engine/Plugins/Runtime/nDisplay/Source/DisplayCluster/Private/DisplayClusterModule.h
--- a/Engine/Plugins/Runtime/nDisplay/Source/DisplayCluster/Private/DisplayClusterModule.h
+++ b/Engine/Plugins/Runtime/nDisplay/Source/DisplayCluster/Private/DisplayClusterModule.h
@@ -25,6 +25,12 @@ public:
 	//////////////////////////////////////////////////////////////////////////////////////////////
 	// IDisplayCluster
 	//////////////////////////////////////////////////////////////////////////////////////////////
+	virtual bool IsModuleInitialized() const override
+	{ return bIsModuleInitialized; }
+	
+	virtual EDisplayClusterOperationMode GetOperationMode() const override
+	{ return CurrentOperationMode; }
+
 	virtual IDisplayClusterRenderManager*    GetRenderMgr()    const override { return MgrRender; }
 	virtual IDisplayClusterClusterManager*   GetClusterMgr()   const override { return MgrCluster; }
 	virtual IDisplayClusterInputManager*     GetInputMgr()     const override { return MgrInput; }
@@ -41,9 +47,6 @@ public:
 	virtual IPDisplayClusterConfigManager*    GetPrivateConfigMgr()    const override { return MgrConfig; }
 	virtual IPDisplayClusterGameManager*      GetPrivateGameMgr()      const override { return MgrGame; }
 
-	virtual EDisplayClusterOperationMode       GetOperationMode() const override
-	{ return CurrentOperationMode; }
-
 public:
 	//////////////////////////////////////////////////////////////////////////////////////////////
 	// IPDisplayClusterManager
@@ -71,6 +74,10 @@ private:
 #endif
 
 private:
+	// Is module initialized.
+	// This flag is not the same as EDisplayClusterOperationMode::Disabled which is used when we turn off the DC functionality in a game mode.
+	bool bIsModuleInitialized = false;
+
 	// DisplayCluster subsystems
 	IPDisplayClusterClusterManager*   MgrCluster   = nullptr;
 	IPDisplayClusterRenderManager*    MgrRender    = nullptr;
diff --git a/Engine/Plugins/Runtime/nDisplay/Source/DisplayCluster/Private/Game/Classes/Basics/DisplayClusterGameMode.cpp b/Engine/Plugins/Runtime/nDisplay/Source/DisplayCluster/Private/Game/Classes/Basics/DisplayClusterGameMode.cpp
--- a/Engine/Plugins/Runtime/nDisplay/Source/DisplayCluster/Private/Game/Classes/Basics/DisplayClusterGameMode.cpp
+++ b/Engine/Plugins/Runtime/nDisplay/Source/DisplayCluster/Private/Game/Classes/Basics/DisplayClusterGameMode.cpp
@@ -52,7 +52,12 @@ void ADisplayClusterGameMode::InitGame(const FString& MapName, const FString& Op
 
 	Super::InitGame(MapName, Options, ErrorMessage);
 
-	UE_LOG(LogDisplayClusterGame, Log, TEXT("%s"), bIsDisplayClusterActive ? 
+	if (!GDisplayCluster->IsModuleInitialized())
+	{
+		return;
+	}
+
+	UE_LOG(LogDisplayClusterGame, Log, TEXT("%s"), bIsDisplayClusterActive ?
 		TEXT("DisplayCluster feature is active for this world.") : 
 		TEXT("DisplayCluster feature has been deactivated for this world by game mode."));
 
@@ -117,7 +122,7 @@ void ADisplayClusterGameMode::StartPlay()
 {
 	DISPLAY_CLUSTER_FUNC_TRACE(LogDisplayClusterGame);
 
-	if (bIsDisplayClusterActive)
+	if (GDisplayCluster->IsModuleInitialized() && bIsDisplayClusterActive)
 	{
 		IPDisplayClusterGameManager* const pGameMgr = GDisplayCluster->GetPrivateGameMgr();
 		if (pGameMgr)
@@ -145,7 +150,7 @@ void ADisplayClusterGameMode::BeginPlay()
 {
 	DISPLAY_CLUSTER_FUNC_TRACE(LogDisplayClusterGame);
 
-	if (bIsDisplayClusterActive)
+	if (GDisplayCluster->IsModuleInitialized() && bIsDisplayClusterActive)
 	{
 		bGameStarted = GDisplayCluster->StartScene(GetWorld());
 		if (!bGameStarted)
@@ -162,7 +167,7 @@ void ADisplayClusterGameMode::BeginDestroy()
 {
 	DISPLAY_CLUSTER_FUNC_TRACE(LogDisplayClusterGame);
 
-	if (bIsDisplayClusterActive)
+	if (GDisplayCluster->IsModuleInitialized() && bIsDisplayClusterActive)
 	{
 		if (bGameStarted)
 		{
@@ -181,7 +186,7 @@ void ADisplayClusterGameMode::Tick(float DeltaSeconds)
 
 	Super::Tick(DeltaSeconds);
 
-	if (!bIsDisplayClusterActive)
+	if (!GDisplayCluster->IsModuleInitialized() || !bIsDisplayClusterActive)
 	{
 		return;
 	}
diff --git a/Engine/Plugins/Runtime/nDisplay/Source/DisplayCluster/Private/Game/Classes/Basics/DisplayClusterGameModeDefault.cpp b/Engine/Plugins/Runtime/nDisplay/Source/DisplayCluster/Private/Game/Classes/Basics/DisplayClusterGameModeDefault.cpp
--- a/Engine/Plugins/Runtime/nDisplay/Source/DisplayCluster/Private/Game/Classes/Basics/DisplayClusterGameModeDefault.cpp
+++ b/Engine/Plugins/Runtime/nDisplay/Source/DisplayCluster/Private/Game/Classes/Basics/DisplayClusterGameModeDefault.cpp
@@ -4,6 +4,8 @@
 #include "DisplayClusterPawnDefault.h"
 
 #include "Misc/DisplayClusterLog.h"
+#include "IPDisplayCluster.h"
+#include "DisplayClusterGlobals.h"
 
 
 ADisplayClusterGameModeDefault::ADisplayClusterGameModeDefault() :
diff --git a/Engine/Plugins/Runtime/nDisplay/Source/DisplayCluster/Private/Game/Classes/Scene/DisplayClusterPawn.cpp b/Engine/Plugins/Runtime/nDisplay/Source/DisplayCluster/Private/Game/Classes/Scene/DisplayClusterPawn.cpp
--- a/Engine/Plugins/Runtime/nDisplay/Source/DisplayCluster/Private/Game/Classes/Scene/DisplayClusterPawn.cpp
+++ b/Engine/Plugins/Runtime/nDisplay/Source/DisplayCluster/Private/Game/Classes/Scene/DisplayClusterPawn.cpp
@@ -67,6 +67,11 @@ void ADisplayClusterPawn::BeginPlay()
 
 	Super::BeginPlay();
 
+	if (!GDisplayCluster->IsModuleInitialized())
+	{
+		return;
+	}
+
 	GameMgr = GDisplayCluster->GetPrivateGameMgr();
 	bIsCluster = (GDisplayCluster->GetOperationMode() == EDisplayClusterOperationMode::Cluster);
 
@@ -74,11 +79,12 @@ void ADisplayClusterPawn::BeginPlay()
 	CollisionComponent->SetCollisionEnabled(ECollisionEnabled::NoCollision);
 
 	// Enable collision if needed
-	if (GameMgr->IsDisplayClusterActive())
+	if (GameMgr && GameMgr->IsDisplayClusterActive())
 	{
 		const ADisplayClusterSettings* const pDisplayClusterSettings = GameMgr->GetDisplayClusterSceneSettings();
 
-		if (GDisplayCluster->GetPrivateClusterMgr()->IsMaster())
+		const IPDisplayClusterClusterManager* const ClusterMgr = GDisplayCluster->GetPrivateClusterMgr();
+		if (ClusterMgr && ClusterMgr->IsMaster())
 		{
 			if (pDisplayClusterSettings && pDisplayClusterSettings->bEnableCollisions)
 			{
diff --git a/Engine/Plugins/Runtime/nDisplay/Source/DisplayCluster/Private/Game/Classes/Scene/DisplayClusterPawnDefault.cpp b/Engine/Plugins/Runtime/nDisplay/Source/DisplayCluster/Private/Game/Classes/Scene/DisplayClusterPawnDefault.cpp
--- a/Engine/Plugins/Runtime/nDisplay/Source/DisplayCluster/Private/Game/Classes/Scene/DisplayClusterPawnDefault.cpp
+++ b/Engine/Plugins/Runtime/nDisplay/Source/DisplayCluster/Private/Game/Classes/Scene/DisplayClusterPawnDefault.cpp
@@ -69,6 +69,11 @@ void ADisplayClusterPawnDefault::BeginPlay()
 
 	Super::BeginPlay();
 
+	if (!GDisplayCluster->IsModuleInitialized())
+	{
+		return;
+	}
+
 	GameMgr = GDisplayCluster->GetPrivateGameMgr();
 	bIsCluster = (GDisplayCluster->GetOperationMode() == EDisplayClusterOperationMode::Cluster);
 
@@ -77,7 +82,7 @@ void ADisplayClusterPawnDefault::BeginPlay()
 	bUseControllerRotationRoll  = !bIsCluster;
 
 	// Enable collision if needed
-	if (GameMgr->IsDisplayClusterActive())
+	if (GameMgr && GameMgr->IsDisplayClusterActive())
 	{
 		const ADisplayClusterSettings* const pDisplayClusterSettings = GameMgr->GetDisplayClusterSceneSettings();
 		if (pDisplayClusterSettings)
@@ -223,6 +228,7 @@ void ADisplayClusterPawnDefault::LookUpAtRate(float Rate)
 
 	if (bIsCluster)
 	{
+		//@note: usually CAVE-like systems don't use roll and pitch rotation since it can cause dizziness.
 #if 0
 		//@todo: rotate around active camera
 		IPDisplayClusterGameManager* const pMgr = GDisplayCluster->GetPrivateGameMgr();
diff --git a/Engine/Plugins/Runtime/nDisplay/Source/DisplayCluster/Private/Game/Classes/Scene/DisplayClusterSceneComponent.cpp b/Engine/Plugins/Runtime/nDisplay/Source/DisplayCluster/Private/Game/Classes/Scene/DisplayClusterSceneComponent.cpp
--- a/Engine/Plugins/Runtime/nDisplay/Source/DisplayCluster/Private/Game/Classes/Scene/DisplayClusterSceneComponent.cpp
+++ b/Engine/Plugins/Runtime/nDisplay/Source/DisplayCluster/Private/Game/Classes/Scene/DisplayClusterSceneComponent.cpp
@@ -34,13 +34,13 @@ void UDisplayClusterSceneComponent::TickComponent( float DeltaTime, ELevelTick T
 	// Update transform if attached to a tracker
 	if (!Config.TrackerId.IsEmpty())
 	{
-		IPDisplayClusterInputManager* const inputMgr = GDisplayCluster->GetPrivateInputMgr();
-		if (inputMgr)
+		const IPDisplayClusterInputManager* const InputMgr = GDisplayCluster->GetPrivateInputMgr();
+		if (InputMgr)
 		{
 			FVector loc;
 			FQuat rot;
-			const bool bLocAvail = inputMgr->GetTrackerLocation(Config.TrackerId, Config.TrackerCh, loc);
-			const bool bRotAvail = inputMgr->GetTrackerQuat(Config.TrackerId, Config.TrackerCh, rot);
+			const bool bLocAvail = InputMgr->GetTrackerLocation(Config.TrackerId, Config.TrackerCh, loc);
+			const bool bRotAvail = InputMgr->GetTrackerQuat(Config.TrackerId, Config.TrackerCh, rot);
 
 			if (bLocAvail && bRotAvail)
 			{
@@ -71,10 +71,14 @@ bool UDisplayClusterSceneComponent::ApplySettings()
 	// Take place in hierarchy
 	if (!GetParentId().IsEmpty())
 	{
-		UE_LOG(LogDisplayClusterGame, Log, TEXT("Attaching %s to %s"), *GetId(), *GetParentId());
-		UDisplayClusterSceneComponent* const pComp = GDisplayCluster->GetPrivateGameMgr()->GetNodeById(GetParentId());
-		AttachToComponent(pComp, FAttachmentTransformRules(EAttachmentRule::KeepRelative, false));
-		//this->SetRelativeTransform(FTransform::Identity);
+		const IPDisplayClusterGameManager* const GameMgr = GDisplayCluster->GetPrivateGameMgr();
+		if (GameMgr)
+		{
+			UE_LOG(LogDisplayClusterGame, Log, TEXT("Attaching %s to %s"), *GetId(), *GetParentId());
+			UDisplayClusterSceneComponent* const pComp = GameMgr->GetNodeById(GetParentId());
+			AttachToComponent(pComp, FAttachmentTransformRules(EAttachmentRule::KeepRelative, false));
+			//this->SetRelativeTransform(FTransform::Identity);
+		}
 	}
 
 	// Set up location and rotation
diff --git a/Engine/Plugins/Runtime/nDisplay/Source/DisplayCluster/Private/Game/Classes/Scene/DisplayClusterSceneComponentSync.cpp b/Engine/Plugins/Runtime/nDisplay/Source/DisplayCluster/Private/Game/Classes/Scene/DisplayClusterSceneComponentSync.cpp
--- a/Engine/Plugins/Runtime/nDisplay/Source/DisplayCluster/Private/Game/Classes/Scene/DisplayClusterSceneComponentSync.cpp
+++ b/Engine/Plugins/Runtime/nDisplay/Source/DisplayCluster/Private/Game/Classes/Scene/DisplayClusterSceneComponentSync.cpp
@@ -19,12 +19,16 @@ void UDisplayClusterSceneComponentSync::BeginPlay()
 {
 	Super::BeginPlay();
 
-	GameMgr = GDisplayCluster->GetPrivateGameMgr();
+	if (!GDisplayCluster->IsModuleInitialized())
+	{
+		return;
+	}
 
 	// Generate unique sync id
 	SyncId = GetSyncId();
 
-	if (GameMgr->IsDisplayClusterActive())
+	GameMgr = GDisplayCluster->GetPrivateGameMgr();
+	if (GameMgr && GameMgr->IsDisplayClusterActive())
 	{
 		// Register sync object
 		ClusterMgr = GDisplayCluster->GetPrivateClusterMgr();
@@ -50,10 +54,16 @@ void UDisplayClusterSceneComponentSync::TickComponent( float DeltaTime, ELevelTi
 
 void UDisplayClusterSceneComponentSync::DestroyComponent(bool bPromoteChildren)
 {
-	if (ClusterMgr)
+	if (GDisplayCluster->IsModuleInitialized())
 	{
-		UE_LOG(LogDisplayClusterGame, Log, TEXT("Unregistering sync object %s..."), *SyncId);
-		ClusterMgr->UnregisterSyncObject(this);
+		if (GameMgr && GameMgr->IsDisplayClusterActive())
+		{
+			if (ClusterMgr)
+			{
+				UE_LOG(LogDisplayClusterGame, Log, TEXT("Unregistering sync object %s..."), *SyncId);
+				ClusterMgr->UnregisterSyncObject(this);
+			}
+		}
 	}
 
 	Super::DestroyComponent(bPromoteChildren);
diff --git a/Engine/Plugins/Runtime/nDisplay/Source/DisplayCluster/Private/Game/Classes/Scene/DisplayClusterScreenComponent.cpp b/Engine/Plugins/Runtime/nDisplay/Source/DisplayCluster/Private/Game/Classes/Scene/DisplayClusterScreenComponent.cpp
--- a/Engine/Plugins/Runtime/nDisplay/Source/DisplayCluster/Private/Game/Classes/Scene/DisplayClusterScreenComponent.cpp
+++ b/Engine/Plugins/Runtime/nDisplay/Source/DisplayCluster/Private/Game/Classes/Scene/DisplayClusterScreenComponent.cpp
@@ -23,22 +23,26 @@ UDisplayClusterScreenComponent::UDisplayClusterScreenComponent(const FObjectInit
 #if WITH_EDITOR
 	if (GEngine && GEngine->IsEditor())
 	{
-		const ADisplayClusterSettings* const pDisplayClusterSettings = GDisplayCluster->GetPrivateGameMgr()->GetDisplayClusterSceneSettings();
-		if(pDisplayClusterSettings && pDisplayClusterSettings->bEditorShowProjectionScreens)
+		const IPDisplayClusterGameManager* const GameMgr = GDisplayCluster->GetPrivateGameMgr();
+		if (GameMgr)
 		{
-			ScreenGeometryComponent = CreateDefaultSubobject<UStaticMeshComponent>(FName(*(GetName() + FString("_impl"))));
-			check(ScreenGeometryComponent);
-
-			if (ScreenGeometryComponent)
+			const ADisplayClusterSettings* const pDisplayClusterSettings = GameMgr->GetDisplayClusterSceneSettings();
+			if (pDisplayClusterSettings && pDisplayClusterSettings->bEditorShowProjectionScreens)
 			{
-				static ConstructorHelpers::FObjectFinder<UStaticMesh> screenMesh(TEXT("StaticMesh'/Engine/BasicShapes/Cube.Cube'"));
-				static ConstructorHelpers::FObjectFinder<UMaterial>   screenMat(TEXT("Material'/Engine/Engine_MI_Shaders/M_Shader_SimpleTranslucent.M_Shader_SimpleTranslucent'"));
-
-				ScreenGeometryComponent->AttachToComponent(this, FAttachmentTransformRules(EAttachmentRule::KeepRelative, false));
-				ScreenGeometryComponent->SetStaticMesh(screenMesh.Object);
-				ScreenGeometryComponent->SetMobility(EComponentMobility::Movable);
-				ScreenGeometryComponent->SetMaterial(0, screenMat.Object);
-				ScreenGeometryComponent->SetCollisionEnabled(ECollisionEnabled::NoCollision);
+				ScreenGeometryComponent = CreateDefaultSubobject<UStaticMeshComponent>(FName(*(GetName() + FString("_impl"))));
+				check(ScreenGeometryComponent);
+
+				if (ScreenGeometryComponent)
+				{
+					static ConstructorHelpers::FObjectFinder<UStaticMesh> screenMesh(TEXT("StaticMesh'/Engine/BasicShapes/Cube.Cube'"));
+					static ConstructorHelpers::FObjectFinder<UMaterial>   screenMat(TEXT("Material'/Engine/Engine_MI_Shaders/M_Shader_SimpleTranslucent.M_Shader_SimpleTranslucent'"));
+
+					ScreenGeometryComponent->AttachToComponent(this, FAttachmentTransformRules(EAttachmentRule::KeepRelative, false));
+					ScreenGeometryComponent->SetStaticMesh(screenMesh.Object);
+					ScreenGeometryComponent->SetMobility(EComponentMobility::Movable);
+					ScreenGeometryComponent->SetMaterial(0, screenMat.Object);
+					ScreenGeometryComponent->SetCollisionEnabled(ECollisionEnabled::NoCollision);
+				}
 			}
 		}
 	}
diff --git a/Engine/Plugins/Runtime/nDisplay/Source/DisplayCluster/Private/IPDisplayCluster.h b/Engine/Plugins/Runtime/nDisplay/Source/DisplayCluster/Private/IPDisplayCluster.h
--- a/Engine/Plugins/Runtime/nDisplay/Source/DisplayCluster/Private/IPDisplayCluster.h
+++ b/Engine/Plugins/Runtime/nDisplay/Source/DisplayCluster/Private/IPDisplayCluster.h
@@ -30,6 +30,4 @@ struct IPDisplayCluster
 	virtual IPDisplayClusterInputManager*     GetPrivateInputMgr() const = 0;
 	virtual IPDisplayClusterConfigManager*    GetPrivateConfigMgr() const = 0;
 	virtual IPDisplayClusterGameManager*      GetPrivateGameMgr() const = 0;
-
-	virtual EDisplayClusterOperationMode       GetOperationMode() const = 0;
 };
diff --git a/Engine/Plugins/Runtime/nDisplay/Source/DisplayCluster/Private/Input/DisplayClusterInputManager.cpp b/Engine/Plugins/Runtime/nDisplay/Source/DisplayCluster/Private/Input/DisplayClusterInputManager.cpp
--- a/Engine/Plugins/Runtime/nDisplay/Source/DisplayCluster/Private/Input/DisplayClusterInputManager.cpp
+++ b/Engine/Plugins/Runtime/nDisplay/Source/DisplayCluster/Private/Input/DisplayClusterInputManager.cpp
@@ -233,6 +233,11 @@ void FDisplayClusterInputManager::Update()
 {
 	DISPLAY_CLUSTER_FUNC_TRACE(LogDisplayClusterInput);
 
+	if (GDisplayCluster->GetOperationMode() == EDisplayClusterOperationMode::Disabled)
+	{
+		return;
+	}
+
 	// Perform input update on master only. Slaves' state will be replicated later.
 	if (GDisplayCluster->GetPrivateClusterMgr()->IsMaster())
 	{
@@ -376,6 +381,11 @@ bool FDisplayClusterInputManager::InitDevices()
 {
 	DISPLAY_CLUSTER_FUNC_TRACE(LogDisplayClusterInput);
 
+	if (GDisplayCluster->GetOperationMode() == EDisplayClusterOperationMode::Disabled)
+	{
+		return false;
+	}
+
 	FScopeLock ScopeLock(&InternalsSyncScope);
 
 	UE_LOG(LogDisplayClusterInput, Log, TEXT("Initializing input devices..."));
diff --git a/Engine/Plugins/Runtime/nDisplay/Source/DisplayCluster/Private/Network/Service/DisplayClusterService.cpp b/Engine/Plugins/Runtime/nDisplay/Source/DisplayCluster/Private/Network/Service/DisplayClusterService.cpp
--- a/Engine/Plugins/Runtime/nDisplay/Source/DisplayCluster/Private/Network/Service/DisplayClusterService.cpp
+++ b/Engine/Plugins/Runtime/nDisplay/Source/DisplayCluster/Private/Network/Service/DisplayClusterService.cpp
@@ -18,6 +18,11 @@ FDisplayClusterService::FDisplayClusterService(const FString& name, const FStrin
 
 bool FDisplayClusterService::IsClusterIP(const FIPv4Endpoint& ep)
 {
+	if (GDisplayCluster->GetOperationMode() == EDisplayClusterOperationMode::Disabled)
+	{
+		return false;
+	}
+
 	TArray<FDisplayClusterConfigClusterNode> nodes = GDisplayCluster->GetPrivateConfigMgr()->GetClusterNodes();
 	const FString addr = ep.Address.ToString();
 	
diff --git a/Engine/Plugins/Runtime/nDisplay/Source/DisplayCluster/Private/Render/Devices/DisplayClusterDeviceBase.cpp b/Engine/Plugins/Runtime/nDisplay/Source/DisplayCluster/Private/Render/Devices/DisplayClusterDeviceBase.cpp
--- a/Engine/Plugins/Runtime/nDisplay/Source/DisplayCluster/Private/Render/Devices/DisplayClusterDeviceBase.cpp
+++ b/Engine/Plugins/Runtime/nDisplay/Source/DisplayCluster/Private/Render/Devices/DisplayClusterDeviceBase.cpp
@@ -34,6 +34,11 @@ FDisplayClusterDeviceBase::~FDisplayClusterDeviceBase()
 
 bool FDisplayClusterDeviceBase::Initialize()
 {
+	if (GDisplayCluster->GetOperationMode() == EDisplayClusterOperationMode::Disabled)
+	{
+		return false;
+	}
+
 	const FDisplayClusterConfigCustom cfgCustom = GDisplayCluster->GetPrivateConfigMgr()->GetConfigCustom();
 
 	if (cfgCustom.Args.Contains(FString(DisplayClusterStrings::cfg::data::custom::SwapInt)))
@@ -75,6 +80,11 @@ void FDisplayClusterDeviceBase::UpdateProjectionScreenDataForThisFrame()
 	UE_LOG(LogDisplayClusterRender, VeryVerbose, TEXT("UpdateProjectionScreenDataForThisFrame"));
 	check(IsInGameThread());
 
+	if (GDisplayCluster->GetOperationMode() == EDisplayClusterOperationMode::Disabled)
+	{
+		return;
+	}
+
 	// Store transformations of active projection screen
 	UDisplayClusterScreenComponent* pScreen = GDisplayCluster->GetPrivateGameMgr()->GetActiveScreen();
 	if (pScreen)
@@ -87,6 +97,11 @@ void FDisplayClusterDeviceBase::UpdateProjectionScreenDataForThisFrame()
 
 void FDisplayClusterDeviceBase::exec_BarrierWait()
 {
+	if (GDisplayCluster->GetOperationMode() == EDisplayClusterOperationMode::Disabled)
+	{
+		return;
+	}
+
 	double tTime = 0.f;
 	double bTime = 0.f;
 
diff --git a/Engine/Plugins/Runtime/nDisplay/Source/DisplayCluster/Private/Render/DisplayClusterRenderManager.h b/Engine/Plugins/Runtime/nDisplay/Source/DisplayCluster/Private/Render/DisplayClusterRenderManager.h
--- a/Engine/Plugins/Runtime/nDisplay/Source/DisplayCluster/Private/Render/DisplayClusterRenderManager.h
+++ b/Engine/Plugins/Runtime/nDisplay/Source/DisplayCluster/Private/Render/DisplayClusterRenderManager.h
@@ -6,7 +6,6 @@
 #include "IPDisplayClusterRenderManager.h"
 
 class FDisplayClusterDeviceBase;
-enum class EDisplayClusterOperationMode;
 
 
 /**
diff --git a/Engine/Plugins/Runtime/nDisplay/Source/DisplayCluster/Public/Blueprints/IDisplayClusterBlueprintAPI.h b/Engine/Plugins/Runtime/nDisplay/Source/DisplayCluster/Public/Blueprints/IDisplayClusterBlueprintAPI.h
--- a/Engine/Plugins/Runtime/nDisplay/Source/DisplayCluster/Public/Blueprints/IDisplayClusterBlueprintAPI.h
+++ b/Engine/Plugins/Runtime/nDisplay/Source/DisplayCluster/Public/Blueprints/IDisplayClusterBlueprintAPI.h
@@ -4,6 +4,7 @@
 
 #include "CoreMinimal.h"
 #include "UObject/Interface.h"
+#include "DisplayClusterOperationMode.h"
 #include "IDisplayClusterBlueprintAPI.generated.h"
 
 
@@ -21,6 +22,18 @@ class DISPLAYCLUSTER_API IDisplayClusterBlueprintAPI
 {
 	GENERATED_BODY()
 
+public:
+	//////////////////////////////////////////////////////////////////////////////////////////////
+	// DisplayCluster module API
+	//////////////////////////////////////////////////////////////////////////////////////////////
+	/** Return if the module has been initialized. */
+	UFUNCTION(BlueprintCallable, meta = (DisplayName = "Is module initialized"), Category = "DisplayCluster")
+	virtual bool IsModuleInitialized() = 0;
+
+	/** Return current operation mode. */
+	UFUNCTION(BlueprintCallable, meta = (DisplayName = "Get operation mode"), Category = "DisplayCluster")
+	virtual EDisplayClusterOperationMode GetOperationMode() = 0;
+
 public:
 	//////////////////////////////////////////////////////////////////////////////////////////////
 	// Cluster API
diff --git a/Engine/Plugins/Runtime/nDisplay/Source/DisplayCluster/Public/DisplayClusterOperationMode.h b/Engine/Plugins/Runtime/nDisplay/Source/DisplayCluster/Public/DisplayClusterOperationMode.h
--- a/Engine/Plugins/Runtime/nDisplay/Source/DisplayCluster/Public/DisplayClusterOperationMode.h
+++ b/Engine/Plugins/Runtime/nDisplay/Source/DisplayCluster/Public/DisplayClusterOperationMode.h
@@ -2,11 +2,14 @@
 
 #pragma once
 
+#include "CoreMinimal.h"
+#include "DisplayClusterOperationMode.generated.h"
 
 /**
  * Display cluster operation mode
  */
-enum class EDisplayClusterOperationMode
+UENUM(BlueprintType)
+enum class EDisplayClusterOperationMode : uint8
 {
 	Cluster,
 	Standalone,
diff --git a/Engine/Plugins/Runtime/nDisplay/Source/DisplayCluster/Public/DisplayClusterPawn.h b/Engine/Plugins/Runtime/nDisplay/Source/DisplayCluster/Public/DisplayClusterPawn.h
--- a/Engine/Plugins/Runtime/nDisplay/Source/DisplayCluster/Public/DisplayClusterPawn.h
+++ b/Engine/Plugins/Runtime/nDisplay/Source/DisplayCluster/Public/DisplayClusterPawn.h
@@ -65,7 +65,10 @@ protected:
 	UDisplayClusterSceneComponent* CollisionOffsetComponent;
 
 private:
+	UPROPERTY()
 	UDisplayClusterSceneComponentSyncParent* DisplayClusterSyncRoot;
+	
+	UPROPERTY()
 	UDisplayClusterSceneComponentSyncParent* DisplayClusterSyncCollisionOffset;
 
 	IPDisplayClusterGameManager* GameMgr = nullptr;
diff --git a/Engine/Plugins/Runtime/nDisplay/Source/DisplayCluster/Public/DisplayClusterPawnDefault.h b/Engine/Plugins/Runtime/nDisplay/Source/DisplayCluster/Public/DisplayClusterPawnDefault.h
--- a/Engine/Plugins/Runtime/nDisplay/Source/DisplayCluster/Public/DisplayClusterPawnDefault.h
+++ b/Engine/Plugins/Runtime/nDisplay/Source/DisplayCluster/Public/DisplayClusterPawnDefault.h
@@ -41,9 +41,6 @@ public:
 	virtual void BeginDestroy() override;
 	virtual void Tick(float DeltaSeconds) override;
 
-protected:
-	virtual void SetupPlayerInputComponent(UInputComponent* InInputComponent) override;
-
 public:
 	/**
 	* Input callback to move forward in local space (or backward if Val is negative).
@@ -86,6 +83,9 @@ public:
 	UFUNCTION(BlueprintCallable, Category = "Pawn")
 	void LookUpAtRate(float Rate);
 
+protected:
+	virtual void SetupPlayerInputComponent(UInputComponent* InInputComponent) override;
+
 protected:
 	/** Movement component */
 	UPROPERTY(Category = Pawn, VisibleAnywhere, BlueprintReadOnly, meta = (AllowPrivateAccess = "true"))
@@ -101,5 +101,5 @@ protected:
 private:
 	IPDisplayClusterGameManager* GameMgr = nullptr;
 
-	bool bIsCluster;
+	bool bIsCluster = false;
 };
diff --git a/Engine/Plugins/Runtime/nDisplay/Source/DisplayCluster/Public/IDisplayCluster.h b/Engine/Plugins/Runtime/nDisplay/Source/DisplayCluster/Public/IDisplayCluster.h
--- a/Engine/Plugins/Runtime/nDisplay/Source/DisplayCluster/Public/IDisplayCluster.h
+++ b/Engine/Plugins/Runtime/nDisplay/Source/DisplayCluster/Public/IDisplayCluster.h
@@ -5,6 +5,9 @@
 #include "Modules/ModuleManager.h"
 #include "Modules/ModuleInterface.h"
 
+#include "DisplayClusterOperationMode.h"
+
+
 struct IDisplayClusterRenderManager;
 struct IDisplayClusterClusterManager;
 struct IDisplayClusterInputManager;
@@ -44,6 +47,21 @@ struct IDisplayCluster
 		return FModuleManager::Get().IsModuleLoaded(IDisplayCluster::ModuleName);
 	}
 
+
+	/**
+	* Checks if the module has been initialized.
+	*
+	* @return Is initialized
+	*/
+	virtual bool IsModuleInitialized() const = 0;
+
+	/**
+	* Returns current operation mode.
+	*
+	* @return Display Cluster operation mode
+	*/
+	virtual EDisplayClusterOperationMode GetOperationMode() const = 0;
+
 	/**
 	* Access to the device manager.
 	*
diff --git a/Engine/Shaders/Private/LocalVertexFactory.ush b/Engine/Shaders/Private/LocalVertexFactory.ush
--- a/Engine/Shaders/Private/LocalVertexFactory.ush
+++ b/Engine/Shaders/Private/LocalVertexFactory.ush
@@ -636,15 +636,9 @@ FVertexFactoryIntermediates GetVertexFactoryIntermediates(FVertexFactoryInput In
 	Intermediates.InstanceOrigin = VertexFetch_InstanceOriginBuffer[(Input.InstanceId + InstanceOffset)];
 	Intermediates.InstanceLightmapAndShadowMapUVBias = VertexFetch_InstanceLightmapBuffer[(Input.InstanceId + InstanceOffset)];
 #elif USE_INSTANCING && !USE_INSTANCING_EMULATED
-	#if WEBGL
-		Intermediates.InstanceTransform1 = f16tof32(Input.InstanceTransform1);
-		Intermediates.InstanceTransform2 = f16tof32(Input.InstanceTransform2);
-		Intermediates.InstanceTransform3 = f16tof32(Input.InstanceTransform3);
-	#else
-		Intermediates.InstanceTransform1 = Input.InstanceTransform1;
-		Intermediates.InstanceTransform2 = Input.InstanceTransform2;
-		Intermediates.InstanceTransform3 = Input.InstanceTransform3;
-	#endif
+	Intermediates.InstanceTransform1 = Input.InstanceTransform1;
+	Intermediates.InstanceTransform2 = Input.InstanceTransform2;
+	Intermediates.InstanceTransform3 = Input.InstanceTransform3;
 	Intermediates.InstanceOrigin = Input.InstanceOrigin;
 	Intermediates.InstanceLightmapAndShadowMapUVBias = Input.InstanceLightmapAndShadowMapUVBias;
 #endif
diff --git a/Engine/Source/Developer/HTML5/HTML5TargetPlatform/Private/HTML5TargetPlatform.h b/Engine/Source/Developer/HTML5/HTML5TargetPlatform/Private/HTML5TargetPlatform.h
--- a/Engine/Source/Developer/HTML5/HTML5TargetPlatform/Private/HTML5TargetPlatform.h
+++ b/Engine/Source/Developer/HTML5/HTML5TargetPlatform/Private/HTML5TargetPlatform.h
@@ -63,6 +63,8 @@ public:
 			return true;
 		case ETargetPlatformFeatures::DeferredRendering:
 			return false;
+		case ETargetPlatformFeatures::HalfFloatVertexFormat:
+			return false;
 		}
 
 		return TTargetPlatformBase<FHTML5PlatformProperties>::SupportsFeature(Feature);
diff --git a/Engine/Source/Developer/TargetPlatform/Public/Common/TargetPlatformBase.h b/Engine/Source/Developer/TargetPlatform/Public/Common/TargetPlatformBase.h
--- a/Engine/Source/Developer/TargetPlatform/Public/Common/TargetPlatformBase.h
+++ b/Engine/Source/Developer/TargetPlatform/Public/Common/TargetPlatformBase.h
@@ -226,6 +226,9 @@ public:
 
 		case ETargetPlatformFeatures::ShouldSplitPaksIntoSmallerSizes :
 			return false;
+
+		case ETargetPlatformFeatures::HalfFloatVertexFormat:
+			return true;
 		}
 
 		return false;
diff --git a/Engine/Source/Developer/TargetPlatform/Public/Interfaces/ITargetPlatform.h b/Engine/Source/Developer/TargetPlatform/Public/Interfaces/ITargetPlatform.h
--- a/Engine/Source/Developer/TargetPlatform/Public/Interfaces/ITargetPlatform.h
+++ b/Engine/Source/Developer/TargetPlatform/Public/Interfaces/ITargetPlatform.h
@@ -62,6 +62,9 @@ enum class ETargetPlatformFeatures
 
 	/* The platform uses software rasterization of the scene for primitive occlusion */
 	SoftwareOcclusion,
+
+	/* The platform supports half float vertex format */
+	HalfFloatVertexFormat,
 };
 
 
diff --git a/Engine/Source/Programs/AutomationTool/Scripts/SyncDepotPath.cs b/Engine/Source/Programs/AutomationTool/Scripts/SyncDepotPath.cs
--- a/Engine/Source/Programs/AutomationTool/Scripts/SyncDepotPath.cs
+++ b/Engine/Source/Programs/AutomationTool/Scripts/SyncDepotPath.cs
@@ -1,4 +1,4 @@
-﻿// Copyright 1998-2018 Epic Games, Inc. All Rights Reserved.
+// Copyright 1998-2018 Epic Games, Inc. All Rights Reserved.
 
 using System;
 using System.Collections.Generic;
@@ -35,6 +35,11 @@ namespace AutomationTool
 			List<KeyValuePair<string, string>> RequiredView = new List<KeyValuePair<string, string>>();
 			RequiredView.Add(new KeyValuePair<string, string>(DepotPath, "/..."));
 
+			if(P4.DoesClientExist(ClientName))
+			{
+				P4.DeleteClient(ClientName);
+			}
+
 			P4ClientInfo Client = new P4ClientInfo();
 			Client.Owner = P4Env.User;
 			Client.Host = Environment.MachineName;
diff --git a/Engine/Source/Runtime/ApplicationCore/Private/Mac/MacApplication.cpp b/Engine/Source/Runtime/ApplicationCore/Private/Mac/MacApplication.cpp
--- a/Engine/Source/Runtime/ApplicationCore/Private/Mac/MacApplication.cpp
+++ b/Engine/Source/Runtime/ApplicationCore/Private/Mac/MacApplication.cpp
@@ -735,9 +735,10 @@ void FMacApplication::ProcessMouseMovedEvent(const FDeferredMacEvent& Event, TSh
 		const EWindowZone::Type Zone = GetCurrentWindowZone(EventWindow.ToSharedRef());
 		bool IsMouseOverTitleBar = (Zone == EWindowZone::TitleBar);
 		const bool IsMovable = IsMouseOverTitleBar || IsEdgeZone(Zone);
+		FCocoaWindow* WindowHandle = EventWindow->GetWindowHandle();
 		MainThreadCall(^{
-			[EventWindow->GetWindowHandle() setMovable:IsMovable];
-			[EventWindow->GetWindowHandle() setMovableByWindowBackground:IsMouseOverTitleBar];
+			[WindowHandle setMovable:IsMovable];
+			[WindowHandle setMovableByWindowBackground:IsMouseOverTitleBar];
 		}, NSDefaultRunLoopMode, false);
 	}
 
diff --git a/Engine/Source/Runtime/ApplicationCore/Private/Mac/MacWindow.cpp b/Engine/Source/Runtime/ApplicationCore/Private/Mac/MacWindow.cpp
--- a/Engine/Source/Runtime/ApplicationCore/Private/Mac/MacWindow.cpp
+++ b/Engine/Source/Runtime/ApplicationCore/Private/Mac/MacWindow.cpp
@@ -332,11 +332,6 @@ void FMacWindow::Show()
 
 		bIsFirstTimeVisible = false;
 
-		MainThreadCall(^{
-			SCOPED_AUTORELEASE_POOL;
-			[WindowHandle orderFrontAndMakeMain:bShouldActivate andKey:bShouldActivate];
-		}, UE4ShowEventMode, true);
-
 		if (bShouldActivate)
 		{
 			// Tell MacApplication to send window deactivate and activate messages to Slate without waiting for Cocoa events.
@@ -347,6 +342,12 @@ void FMacWindow::Show()
 			MacApplication->OnWindowOrderedFront(SharedThis(this));
 		}
 
+		FCocoaWindow* WindowHandleCopy = WindowHandle;
+		MainThreadCall(^{
+			SCOPED_AUTORELEASE_POOL;
+			[WindowHandleCopy orderFrontAndMakeMain:bShouldActivate andKey:bShouldActivate];
+		}, UE4ShowEventMode, true);
+
 		bIsVisible = true;
 	}
 }
diff --git a/Engine/Source/Runtime/Core/Public/Math/IntVector.h b/Engine/Source/Runtime/Core/Public/Math/IntVector.h
--- a/Engine/Source/Runtime/Core/Public/Math/IntVector.h
+++ b/Engine/Source/Runtime/Core/Public/Math/IntVector.h
@@ -6,6 +6,7 @@
 #include "Misc/Crc.h"
 #include "Math/UnrealMathUtility.h"
 #include "Containers/UnrealString.h"
+#include "Serialization/StructuredArchive.h"
 
 /**
  * Structure for integer vectors in 3-d space.
diff --git a/Engine/Source/Runtime/Engine/Classes/Components/InputComponent.h b/Engine/Source/Runtime/Engine/Classes/Components/InputComponent.h
--- a/Engine/Source/Runtime/Engine/Classes/Components/InputComponent.h
+++ b/Engine/Source/Runtime/Engine/Classes/Components/InputComponent.h
@@ -608,6 +608,9 @@ struct FCachedKeyToActionInfo
 	/** Reverse lookup map to speed up evaluation of action bindings. Will be rebuilt via ConditionalBuildKeyMap when action bindings or key mappings are changed. */
 	TMap<FKey, TArray<TSharedPtr<FInputActionBinding>>> KeyToActionMap;
 
+	/** Keep the AnyKey to action map separately as we don't want to have query the map to find it every time. */
+	TArray<TSharedPtr<FInputActionBinding>> AnyKeyToActionMap;
+
 	FCachedKeyToActionInfo()
 		: PlayerInput(nullptr)
 		, KeyMapBuiltForIndex(0)
diff --git a/Engine/Source/Runtime/Engine/Classes/Components/PrimitiveComponent.h b/Engine/Source/Runtime/Engine/Classes/Components/PrimitiveComponent.h
--- a/Engine/Source/Runtime/Engine/Classes/Components/PrimitiveComponent.h
+++ b/Engine/Source/Runtime/Engine/Classes/Components/PrimitiveComponent.h
@@ -269,8 +269,7 @@ public:
 	UFUNCTION(BlueprintSetter)
 	void SetGenerateOverlapEvents(bool bInGenerateOverlapEvents);
 
-//LUMIN_MERGE - Removed private to allow FoundObjectComponent access
-//private:
+private:
 	UPROPERTY(EditAnywhere, BlueprintGetter = GetGenerateOverlapEvents, BlueprintSetter = SetGenerateOverlapEvents, Category = Collision)
 	uint8 bGenerateOverlapEvents : 1;
 
diff --git a/Engine/Source/Runtime/Engine/Private/Components/InputComponent.cpp b/Engine/Source/Runtime/Engine/Private/Components/InputComponent.cpp
--- a/Engine/Source/Runtime/Engine/Private/Components/InputComponent.cpp
+++ b/Engine/Source/Runtime/Engine/Private/Components/InputComponent.cpp
@@ -52,11 +52,12 @@ void UInputComponent::ConditionalBuildKeyMap(UPlayerInput* PlayerInput)
 		CachedInfoToPopulate->PlayerInput = PlayerInput;
 	}
 
-	// Reset the map
+	// Reset the map and AnyKey array
 	for (TPair<FKey, TArray<TSharedPtr<FInputActionBinding>>>& KeyBindingPair : CachedInfoToPopulate->KeyToActionMap)
 	{
 		KeyBindingPair.Value.Reset();
 	}
+	CachedInfoToPopulate->AnyKeyToActionMap.Reset();
 
 	for (const TSharedPtr<FInputActionBinding>& ActionBinding : ActionBindings)
 	{
@@ -64,7 +65,14 @@ void UInputComponent::ConditionalBuildKeyMap(UPlayerInput* PlayerInput)
 
 		for (const FInputActionKeyMapping& KeyMapping : KeysForAction)
 		{
-			CachedInfoToPopulate->KeyToActionMap.FindOrAdd(KeyMapping.Key).Add(ActionBinding);
+			if (KeyMapping.Key != EKeys::AnyKey)
+			{
+				CachedInfoToPopulate->KeyToActionMap.FindOrAdd(KeyMapping.Key).Add(ActionBinding);
+			}
+			else
+			{
+				CachedInfoToPopulate->AnyKeyToActionMap.Add(ActionBinding);
+			}
 		}
 	}
 
@@ -84,6 +92,10 @@ void UInputComponent::GetActionsBoundToKey(UPlayerInput* PlayerInput, const FKey
 					Actions.AddUnique(ActionForKey);
 				}
 			}
+			for (const TSharedPtr<FInputActionBinding>& ActionForKey : CachedInfo.AnyKeyToActionMap)
+			{
+				Actions.AddUnique(ActionForKey);
+			}
 			return;
 		}
 	}
diff --git a/Engine/Source/Runtime/Engine/Private/InstancedStaticMesh.cpp b/Engine/Source/Runtime/Engine/Private/InstancedStaticMesh.cpp
--- a/Engine/Source/Runtime/Engine/Private/InstancedStaticMesh.cpp
+++ b/Engine/Source/Runtime/Engine/Private/InstancedStaticMesh.cpp
@@ -323,6 +323,72 @@ void FStaticMeshInstanceBuffer::BindInstanceVertexBuffer(const class FVertexFact
 }
 
 
+void FStaticMeshInstanceData::Serialize(FArchive& Ar)
+{
+	// HTML5 doesn't support half float so we need to convert at cook time.
+	const bool bCookConvertTransformsToFullFloat = Ar.IsCooking() && bUseHalfFloat && !Ar.CookingTarget()->SupportsFeature(ETargetPlatformFeatures::HalfFloatVertexFormat);
+
+	if (bCookConvertTransformsToFullFloat)
+	{
+		bool bSaveUseHalfFloat = false;
+		Ar << bSaveUseHalfFloat;
+	}
+	else
+	{
+		Ar << bUseHalfFloat;
+	}
+
+	Ar << NumInstances;
+
+	if (Ar.IsLoading())
+	{
+		AllocateBuffers(NumInstances);
+	}
+
+	InstanceOriginData->Serialize(Ar);
+	InstanceLightmapData->Serialize(Ar);
+
+	if (bCookConvertTransformsToFullFloat)
+	{
+		TStaticMeshVertexData<FInstanceTransformMatrix<float>> FullInstanceTransformData;
+		FullInstanceTransformData.ResizeBuffer(NumInstances);
+
+		FInstanceTransformMatrix<FFloat16>* Src = (FInstanceTransformMatrix<FFloat16>*)InstanceTransformData->GetDataPointer();
+		FInstanceTransformMatrix<float>* Dest = (FInstanceTransformMatrix<float>*)FullInstanceTransformData.GetDataPointer();
+		for (int32 Idx = 0; Idx < NumInstances; Idx++)
+		{
+			Dest->InstanceTransform1[0] = Src->InstanceTransform1[0];
+			Dest->InstanceTransform1[1] = Src->InstanceTransform1[1];
+			Dest->InstanceTransform1[2] = Src->InstanceTransform1[2];
+			Dest->InstanceTransform1[3] = Src->InstanceTransform1[3];
+			Dest->InstanceTransform2[0] = Src->InstanceTransform2[0];
+			Dest->InstanceTransform2[1] = Src->InstanceTransform2[1];
+			Dest->InstanceTransform2[2] = Src->InstanceTransform2[2];
+			Dest->InstanceTransform2[3] = Src->InstanceTransform2[3];
+			Dest->InstanceTransform3[0] = Src->InstanceTransform3[0];
+			Dest->InstanceTransform3[1] = Src->InstanceTransform3[1];
+			Dest->InstanceTransform3[2] = Src->InstanceTransform3[2];
+			Dest->InstanceTransform3[3] = Src->InstanceTransform3[3];
+			Src++;
+			Dest++;
+		}
+
+		FullInstanceTransformData.Serialize(Ar);
+	}
+	else
+	{
+		InstanceTransformData->Serialize(Ar);
+	}
+
+	if (Ar.IsLoading())
+	{
+		InstanceOriginDataPtr = InstanceOriginData->GetDataPointer();
+		InstanceLightmapDataPtr = InstanceLightmapData->GetDataPointer();
+		InstanceTransformDataPtr = InstanceTransformData->GetDataPointer();
+	}
+}
+
+
 /**
  * Should we cache the material's shadertype on this platform with this vertex factory? 
  */
diff --git a/Engine/Source/Runtime/Engine/Private/InstancedStaticMesh.h b/Engine/Source/Runtime/Engine/Private/InstancedStaticMesh.h
--- a/Engine/Source/Runtime/Engine/Private/InstancedStaticMesh.h
+++ b/Engine/Source/Runtime/Engine/Private/InstancedStaticMesh.h
@@ -65,9 +65,6 @@ public:
 	ENGINE_API void InitFromPreallocatedData(FStaticMeshInstanceData& Other);
 	ENGINE_API void UpdateFromCommandBuffer_Concurrent(FInstanceUpdateCmdBuffer& CmdBuffer);
 
-	/** Serializer. */
-	friend FArchive& operator<<(FArchive& Ar, FStaticMeshInstanceBuffer& VertexBuffer);
-
 	/**
 	 * Specialized assignment operator, only used when importing LOD's. 
 	 */
diff --git a/Engine/Source/Runtime/Engine/Public/StaticMeshResources.h b/Engine/Source/Runtime/Engine/Public/StaticMeshResources.h
--- a/Engine/Source/Runtime/Engine/Public/StaticMeshResources.h
+++ b/Engine/Source/Runtime/Engine/Public/StaticMeshResources.h
@@ -802,7 +802,7 @@ public:
 	 * @param bInUseHalfFloat - true if device has support for half float in vertex arrays
 	 */
 	FStaticMeshInstanceData(bool bInUseHalfFloat)
-		: bUseHalfFloat(PLATFORM_BUILTIN_VERTEX_HALF_FLOAT || bInUseHalfFloat)
+	:	bUseHalfFloat(PLATFORM_BUILTIN_VERTEX_HALF_FLOAT || bInUseHalfFloat)
 	{
 		AllocateBuffers(0);
 	}
@@ -814,28 +814,8 @@ public:
 		delete InstanceTransformData;
 	}
 
-	void Serialize(FArchive& Ar)
-	{
-		Ar << bUseHalfFloat;
-		Ar << NumInstances;
-		
-		if (Ar.IsLoading())
-		{
-			AllocateBuffers(NumInstances);
-		}
-
-		InstanceOriginData->Serialize(Ar);
-		InstanceLightmapData->Serialize(Ar);
-		InstanceTransformData->Serialize(Ar);
-		
-		if (Ar.IsLoading())
-		{
-			InstanceOriginDataPtr = InstanceOriginData->GetDataPointer();
-			InstanceLightmapDataPtr = InstanceLightmapData->GetDataPointer();
-			InstanceTransformDataPtr = InstanceTransformData->GetDataPointer();
-		}
-	}
-
+	void Serialize(FArchive& Ar);
+	
 	void AllocateInstances(int32 InNumInstances, EResizeBufferFlags BufferFlags, bool DestroyExistingInstances)
 	{
 		NumInstances = InNumInstances;
diff --git a/Engine/Source/Runtime/Renderer/Private/PostProcess/SceneRenderTargets.cpp b/Engine/Source/Runtime/Renderer/Private/PostProcess/SceneRenderTargets.cpp
--- a/Engine/Source/Runtime/Renderer/Private/PostProcess/SceneRenderTargets.cpp
+++ b/Engine/Source/Runtime/Renderer/Private/PostProcess/SceneRenderTargets.cpp
@@ -1853,13 +1853,19 @@ void FSceneRenderTargets::AllocateCommonDepthTargets(FRHICommandList& RHICmdList
 		{
 			const uint32 OldElementSize = SceneDepthZ->ComputeMemorySize();
 			bHMDAllocatedDepthTarget = true;
+		
+			/*
+			@TODO: UE-61597 (4.20) - Right now there appears to be an issue with the shared depth buffer and the shader param caching system.
+			The guess is that since SceneDepthZ's resource is changed under the hood and needs to be reset as a param for certain shaders, but
+			the caching system isn't aware of that. For the time being, we've commented out this block, incurring added expense on the deferred
+			renderer (adding a depth copy that theoretically is unneeded)
 
 			// If SRT and texture are different (MSAA), only modify the resolve render target, to avoid creating a swapchain of MSAA textures
 			if (SceneDepthZ->GetRenderTargetItem().ShaderResourceTexture == SceneDepthZ->GetRenderTargetItem().TargetableTexture)
 			{
 				SceneDepthZ->GetRenderTargetItem().ShaderResourceTexture = SceneDepthZ->GetRenderTargetItem().TargetableTexture = SRTex;
 			}
-			else
+			else*/
 			{
 				SceneDepthZ->GetRenderTargetItem().ShaderResourceTexture = SRTex;
 			}
diff --git a/Engine/Source/Runtime/Slate/Private/Widgets/SInvalidationPanel.cpp b/Engine/Source/Runtime/Slate/Private/Widgets/SInvalidationPanel.cpp
--- a/Engine/Source/Runtime/Slate/Private/Widgets/SInvalidationPanel.cpp
+++ b/Engine/Source/Runtime/Slate/Private/Widgets/SInvalidationPanel.cpp
@@ -356,8 +356,15 @@ int32 SInvalidationPanel::OnPaint( const FPaintArgs& Args, const FGeometry& Allo
 				// we'll have already done one pre-pass before getting here.
 				ChildSlot.GetWidget()->SlatePrepass(AllottedGeometry.Scale);
 
-				FSlateClippingState CurrentClippingState = OutDrawElements.GetClippingManager().GetClippingStates()[OutDrawElements.GetClippingManager().GetClippingIndex()];
-				CachedWindowElements->GetClippingManager().PushClippingState(CurrentClippingState);
+				bool bPushedNewClip = false;
+				const int32 CurrentClippingIndex = OutDrawElements.GetClippingManager().GetClippingIndex();
+
+				if(CurrentClippingIndex != INDEX_NONE)
+				{
+					bPushedNewClip = true;
+					FSlateClippingState CurrentClippingState = OutDrawElements.GetClippingManager().GetClippingStates()[CurrentClippingIndex];
+					CachedWindowElements->GetClippingManager().PushClippingState(CurrentClippingState);
+				}
 
 				//TODO: When SWidget::Paint is called don't drag self if volatile, and we're doing a cache pass.
 				CachedMaxChildLayer = SCompoundWidget::OnPaint(
@@ -369,7 +376,10 @@ int32 SInvalidationPanel::OnPaint( const FPaintArgs& Args, const FGeometry& Allo
 					InWidgetStyle,
 					bParentEnabled);
 
-				CachedWindowElements->GetClippingManager().PopClip();
+				if (bPushedNewClip)
+				{
+					CachedWindowElements->GetClippingManager().PopClip();
+				}
 
 				{
 					CachedResources.Reset();
