Remove the reference to the Settings module used in standalone.
Use proper implement module syntax for ActorLayerUtilitities

diff --git a/Engine/Plugins/Developer/ActorLayerUtilities/Source/ActorLayerUtilities/Private/ActorLayerUtilities.cpp b/Engine/Plugins/Developer/ActorLayerUtilities/Source/ActorLayerUtilities/Private/ActorLayerUtilities.cpp
--- a/Engine/Plugins/Developer/ActorLayerUtilities/Source/ActorLayerUtilities/Private/ActorLayerUtilities.cpp
+++ b/Engine/Plugins/Developer/ActorLayerUtilities/Source/ActorLayerUtilities/Private/ActorLayerUtilities.cpp
@@ -28,4 +28,4 @@ TArray<AActor*> ULayersBlueprintLibrary::GetActors(UObject* WorldContextObject,
 	return AllActors;
 }
 
-IMPLEMENT_MODULE(FDefaultModuleImpl, "ActorLayerUtilities")
\ No newline at end of file
+IMPLEMENT_MODULE(FDefaultModuleImpl, ActorLayerUtilities)
\ No newline at end of file
diff --git a/Engine/Plugins/Developer/ActorLayerUtilities/Source/ActorLayerUtilitiesEditor/Private/ActorLayerUtilitiesEditor.cpp b/Engine/Plugins/Developer/ActorLayerUtilities/Source/ActorLayerUtilitiesEditor/Private/ActorLayerUtilitiesEditor.cpp
--- a/Engine/Plugins/Developer/ActorLayerUtilities/Source/ActorLayerUtilitiesEditor/Private/ActorLayerUtilitiesEditor.cpp
+++ b/Engine/Plugins/Developer/ActorLayerUtilities/Source/ActorLayerUtilitiesEditor/Private/ActorLayerUtilitiesEditor.cpp
@@ -36,4 +36,4 @@ class FActorLayerUtilitiesEditorModule : public IModuleInterface
 };
 
 
-IMPLEMENT_MODULE(FActorLayerUtilitiesEditorModule, "ActorLayerUtilities")
\ No newline at end of file
+IMPLEMENT_MODULE(FActorLayerUtilitiesEditorModule, ActorLayerUtilities)
\ No newline at end of file
diff --git a/Engine/Plugins/Media/MediaFrameworkUtilities/Source/MediaFrameworkUtilities/MediaFrameworkUtilities.Build.cs b/Engine/Plugins/Media/MediaFrameworkUtilities/Source/MediaFrameworkUtilities/MediaFrameworkUtilities.Build.cs
--- a/Engine/Plugins/Media/MediaFrameworkUtilities/Source/MediaFrameworkUtilities/MediaFrameworkUtilities.Build.cs
+++ b/Engine/Plugins/Media/MediaFrameworkUtilities/Source/MediaFrameworkUtilities/MediaFrameworkUtilities.Build.cs
@@ -24,13 +24,17 @@ namespace UnrealBuildTool.Rules
 					"MediaAssets",
 					"MediaIOCore",
 					"MediaUtils",
-					"Settings",
 				});
 
 			if (Target.bBuildEditor == true)
 			{
-				PrivateDependencyModuleNames.Add("UnrealEd");
+				PrivateDependencyModuleNames.AddRange(
+					new string[]
+					{
+						"Settings",
+						"UnrealEd"
+					});
 			}
 		}
 	}
-}
+}
\ No newline at end of file
diff --git a/Engine/Plugins/Media/MediaFrameworkUtilities/Source/MediaFrameworkUtilities/Private/MediaFrameworkUtilitiesModule.cpp b/Engine/Plugins/Media/MediaFrameworkUtilities/Source/MediaFrameworkUtilities/Private/MediaFrameworkUtilitiesModule.cpp
--- a/Engine/Plugins/Media/MediaFrameworkUtilities/Source/MediaFrameworkUtilities/Private/MediaFrameworkUtilitiesModule.cpp
+++ b/Engine/Plugins/Media/MediaFrameworkUtilities/Source/MediaFrameworkUtilities/Private/MediaFrameworkUtilitiesModule.cpp
@@ -3,9 +3,6 @@
 #include "MediaFrameworkUtilitiesModule.h"
 
 #include "Engine/Engine.h"
-#include "Features/IModularFeatures.h"
-#include "ISettingsModule.h"
-#include "ISettingsSection.h"
 #include "Misc/App.h"
 #include "Misc/CoreDelegates.h"
 #include "Modules/ModuleInterface.h"
@@ -14,14 +11,19 @@
 #include "Profile/MediaProfileManager.h"
 #include "Profile/MediaProfileSettings.h"
 
+#if WITH_EDITOR
+#include "ISettingsModule.h"
+#include "ISettingsSection.h"
+#endif //WITH_EDITOR
+
 
 DEFINE_LOG_CATEGORY(LogMediaFrameworkUtilities);
 
 #define LOCTEXT_NAMESPACE "MediaFrameworkUtilities"
 
 /**
- * Implements the MediaFrameworkUtilitiesModule module.
- */
+* Implements the MediaFrameworkUtilitiesModule module.
+*/
 class FMediaFrameworkUtilitiesModule : public IMediaFrameworkUtilitiesModule
 {
 	FMediaProfileManager MediaProfileManager;
@@ -46,6 +48,7 @@ class FMediaFrameworkUtilitiesModule : public IMediaFrameworkUtilitiesModule
 
 	void RegisterSettings()
 	{
+#if WITH_EDITOR
 		if (GIsEditor)
 		{
 			// register settings
@@ -65,10 +68,12 @@ class FMediaFrameworkUtilitiesModule : public IMediaFrameworkUtilitiesModule
 				);
 			}
 		}
+#endif //WITH_EDITOR
 	}
 
 	void UnregisterSettings()
 	{
+#if WITH_EDITOR
 		if (GIsEditor)
 		{
 			// unregister settings
@@ -79,25 +84,26 @@ class FMediaFrameworkUtilitiesModule : public IMediaFrameworkUtilitiesModule
 				SettingsModule->UnregisterSettings("Editor", "Media", "MediaProfile");
 			}
 		}
+#endif //WITH_EDITOR
 	}
 
 	void ApplyStartupMediaProfile()
 	{
 		auto ApplyMediaProfile = [this]()
-			{
-				UMediaProfile* MediaProfile = nullptr;
+		{
+			UMediaProfile* MediaProfile = nullptr;
 
-	#if WITH_EDITOR
-				MediaProfile = GetDefault<UMediaProfileEditorSettings>()->GetUserMediaProfile();
-	#endif
+#if WITH_EDITOR
+			MediaProfile = GetDefault<UMediaProfileEditorSettings>()->GetUserMediaProfile();
+#endif
 
-				if (MediaProfile == nullptr)
-				{
-					MediaProfile = GetDefault<UMediaProfileSettings>()->GetStartupMediaProfile();
-				}
+			if (MediaProfile == nullptr)
+			{
+				MediaProfile = GetDefault<UMediaProfileSettings>()->GetStartupMediaProfile();
+			}
 
-				MediaProfileManager.SetCurrentMediaProfile(MediaProfile);
-			};
+			MediaProfileManager.SetCurrentMediaProfile(MediaProfile);
+		};
 
 		if (FApp::CanEverRender() || GetDefault<UMediaProfileSettings>()->bApplyInCommandlet)
 		{
