Media Framework: Fix merge conflict

diff --git a/Engine/Plugins/Media/MediaFrameworkUtilities/Source/MediaFrameworkUtilitiesEditor/Private/CaptureTab/SMediaFrameworkCaptureOutputWidget.cpp b/Engine/Plugins/Media/MediaFrameworkUtilities/Source/MediaFrameworkUtilitiesEditor/Private/CaptureTab/SMediaFrameworkCaptureOutputWidget.cpp
--- a/Engine/Plugins/Media/MediaFrameworkUtilities/Source/MediaFrameworkUtilitiesEditor/Private/CaptureTab/SMediaFrameworkCaptureOutputWidget.cpp
+++ b/Engine/Plugins/Media/MediaFrameworkUtilities/Source/MediaFrameworkUtilitiesEditor/Private/CaptureTab/SMediaFrameworkCaptureOutputWidget.cpp
@@ -185,7 +185,7 @@ void SMediaFrameworkCaptureCameraViewportWidget::Construct(const FArguments& InA
 		LevelViewportClient->bDisableInput = true;
 
 		// Never allow Matinee to possess these views
-		LevelViewportClient->SetAllowCinematicPreview(false);
+		LevelViewportClient->SetAllowCinematicControl(false);
 
 		// Our viewport is always visible
 		LevelViewportClient->VisibilityDelegate.BindLambda([] {return true; });
