Keep a reference to the current MediaProfile's proxy to prevent GC
Move MediaProfileSettings to default config.

diff --git a/Engine/Plugins/Media/MediaFrameworkUtilities/Source/MediaFrameworkUtilities/Private/Profile/MediaProfile.cpp b/Engine/Plugins/Media/MediaFrameworkUtilities/Source/MediaFrameworkUtilities/Private/Profile/MediaProfile.cpp
--- a/Engine/Plugins/Media/MediaFrameworkUtilities/Source/MediaFrameworkUtilities/Private/Profile/MediaProfile.cpp
+++ b/Engine/Plugins/Media/MediaFrameworkUtilities/Source/MediaFrameworkUtilities/Private/Profile/MediaProfile.cpp
@@ -25,6 +25,12 @@ UMediaSource* UMediaProfile::GetMediaSource(int32 Index) const
 }
 
 
+int32 UMediaProfile::NumMediaSources() const
+{
+	return MediaSources.Num();
+}
+
+
 UMediaOutput* UMediaProfile::GetMediaOutput(int32 Index) const
 {
 	if (MediaOutputs.IsValidIndex(Index))
@@ -35,6 +41,12 @@ UMediaOutput* UMediaProfile::GetMediaOutput(int32 Index) const
 }
 
 
+int32 UMediaProfile::NumMediaOutputs() const
+{
+	return MediaOutputs.Num();
+}
+
+
 UTimecodeProvider* UMediaProfile::GetTimecodeProvider() const
 {
 	return bOverrideTimecodeProvider ? TimecodeProvider : nullptr;
diff --git a/Engine/Plugins/Media/MediaFrameworkUtilities/Source/MediaFrameworkUtilities/Private/Profile/MediaProfileManager.cpp b/Engine/Plugins/Media/MediaFrameworkUtilities/Source/MediaFrameworkUtilities/Private/Profile/MediaProfileManager.cpp
--- a/Engine/Plugins/Media/MediaFrameworkUtilities/Source/MediaFrameworkUtilities/Private/Profile/MediaProfileManager.cpp
+++ b/Engine/Plugins/Media/MediaFrameworkUtilities/Source/MediaFrameworkUtilities/Private/Profile/MediaProfileManager.cpp
@@ -3,6 +3,11 @@
 #include "MediaFrameworkUtilitiesModule.h"
 
 #include "Modules/ModuleManager.h"
+#include "MediaAssets/ProxyMediaOutput.h"
+#include "MediaAssets/ProxyMediaSource.h"
+#include "Profile/MediaProfile.h"
+#include "Profile/MediaProfileSettings.h"
+
 
 IMediaProfileManager& IMediaProfileManager::Get()
 {
@@ -10,14 +15,24 @@ IMediaProfileManager& IMediaProfileManager::Get()
 	return FModuleManager::GetModuleChecked<IMediaFrameworkUtilitiesModule>(NAME_MediaFrameworkUtilities).GetProfileManager();
 }
 
+
+FMediaProfileManager::FMediaProfileManager()
+	: Collector(this)
+	, CurrentMediaProfile(nullptr)
+{
+}
+
+
 UMediaProfile* FMediaProfileManager::GetCurrentMediaProfile() const
 {
-	return CurrentMediaProfile.Get();
+	return CurrentMediaProfile;
 }
 
+
 void FMediaProfileManager::SetCurrentMediaProfile(UMediaProfile* InMediaProfile)
 {
-	UMediaProfile* Previous = CurrentMediaProfile.Get();
+	bool bChanged = false;
+	UMediaProfile* Previous = CurrentMediaProfile;
 	if (InMediaProfile != Previous)
 	{
 		if (Previous)
@@ -30,12 +45,60 @@ void FMediaProfileManager::SetCurrentMediaProfile(UMediaProfile* InMediaProfile)
 			InMediaProfile->Apply();
 		}
 
-		CurrentMediaProfile.Reset(InMediaProfile);
+		// Set Current assets to prevent GC
+		CurrentMediaProfile = InMediaProfile;
+		CurrentProxyMediaSources.Reset();
+		CurrentProxyMediaOutputs.Reset();
+
+		if (CurrentMediaProfile)
+		{
+			TArray<UProxyMediaSource*> SourceProxies = GetDefault<UMediaProfileSettings>()->GetAllMediaSourceProxy();
+			const int32 MinSourceProxies = FMath::Min(SourceProxies.Num(), CurrentMediaProfile->NumMediaSources());
+			for (int32 Index = 0; Index < MinSourceProxies; ++Index)
+			{
+				UProxyMediaSource* Proxy = SourceProxies[Index];
+				if (Proxy)
+				{
+					CurrentProxyMediaSources.Add(Proxy);
+				}
+			}
+		}
+
+		if (CurrentMediaProfile)
+		{
+			TArray<UProxyMediaOutput*> OutputProxies = GetDefault<UMediaProfileSettings>()->GetAllMediaOutputProxy();
+			const int32 MinOutputProxies = FMath::Min(OutputProxies.Num(), CurrentMediaProfile->NumMediaOutputs());
+			for (int32 Index = 0; Index < MinOutputProxies; ++Index)
+			{
+				UProxyMediaOutput* Proxy = OutputProxies[Index];
+				if (Proxy)
+				{
+					CurrentProxyMediaOutputs.Add(Proxy);
+				}
+			}
+		}
+
 		MediaProfileChangedDelegate.Broadcast(Previous, InMediaProfile);
 	}
 }
 
+
 FMediaProfileManager::FOnMediaProfileChanged& FMediaProfileManager::OnMediaProfileChanged()
 {
 	return MediaProfileChangedDelegate;
 }
+
+
+FMediaProfileManager::FInternalReferenceCollector::FInternalReferenceCollector(FMediaProfileManager* InOwner)
+	: Owner(InOwner)
+{
+}
+
+
+void FMediaProfileManager::FInternalReferenceCollector::AddReferencedObjects(FReferenceCollector& InCollector)
+{
+	check(Owner);
+	InCollector.AddReferencedObject(Owner->CurrentMediaProfile);
+	InCollector.AddReferencedObjects(Owner->CurrentProxyMediaSources);
+	InCollector.AddReferencedObjects(Owner->CurrentProxyMediaOutputs);
+}
diff --git a/Engine/Plugins/Media/MediaFrameworkUtilities/Source/MediaFrameworkUtilities/Private/Profile/MediaProfileManager.h b/Engine/Plugins/Media/MediaFrameworkUtilities/Source/MediaFrameworkUtilities/Private/Profile/MediaProfileManager.h
--- a/Engine/Plugins/Media/MediaFrameworkUtilities/Source/MediaFrameworkUtilities/Private/Profile/MediaProfileManager.h
+++ b/Engine/Plugins/Media/MediaFrameworkUtilities/Source/MediaFrameworkUtilities/Private/Profile/MediaProfileManager.h
@@ -3,20 +3,36 @@
 #pragma once
 
 #include "Profile/IMediaProfileManager.h"
+#include "UObject/GCObject.h"
 
-#include "Profile/MediaProfile.h"
-#include "UObject/StrongObjectPtr.h"
+class UMediaProfile;
+class UProxyMediaSource;
+class UProxyMediaOutput;
 
 class FMediaProfileManager : public IMediaProfileManager
 {
 public:
+	FMediaProfileManager();
 
 	virtual UMediaProfile* GetCurrentMediaProfile() const override;
 	virtual void SetCurrentMediaProfile(UMediaProfile* InMediaProfile) override;
 	virtual FOnMediaProfileChanged& OnMediaProfileChanged() override;
 
 private:
+	class FInternalReferenceCollector : public FGCObject
+	{
+	public:
+		FInternalReferenceCollector(FMediaProfileManager* InOwner);
+		virtual void AddReferencedObjects(FReferenceCollector& Collector) override;
+	private:
+		FMediaProfileManager* Owner;
+	};
+	friend FInternalReferenceCollector;
+	FInternalReferenceCollector Collector;
+
+	UMediaProfile* CurrentMediaProfile;
+	TArray<UProxyMediaSource*> CurrentProxyMediaSources;
+	TArray<UProxyMediaOutput*> CurrentProxyMediaOutputs;
 
-	TStrongObjectPtr<UMediaProfile> CurrentMediaProfile;
 	FOnMediaProfileChanged MediaProfileChangedDelegate;
 };
diff --git a/Engine/Plugins/Media/MediaFrameworkUtilities/Source/MediaFrameworkUtilities/Private/Profile/MediaProfileSettings.cpp b/Engine/Plugins/Media/MediaFrameworkUtilities/Source/MediaFrameworkUtilities/Private/Profile/MediaProfileSettings.cpp
--- a/Engine/Plugins/Media/MediaFrameworkUtilities/Source/MediaFrameworkUtilities/Private/Profile/MediaProfileSettings.cpp
+++ b/Engine/Plugins/Media/MediaFrameworkUtilities/Source/MediaFrameworkUtilities/Private/Profile/MediaProfileSettings.cpp
@@ -8,12 +8,6 @@
 #include "Profile/MediaProfile.h"
 
 
-UMediaProfileSettings::UMediaProfileSettings()
-	: bDisplayInToolbar(true)
-{
-}
-
-
 TArray<UProxyMediaSource*> UMediaProfileSettings::GetAllMediaSourceProxy() const
 {
 	TArray<UProxyMediaSource*> Result;
@@ -46,6 +40,12 @@ UMediaProfile* UMediaProfileSettings::GetStartupMediaProfile() const
 }
 
 
+UMediaProfileEditorSettings::UMediaProfileEditorSettings()
+	:bDisplayInToolbar(true)
+{
+}
+
+
 UMediaProfile* UMediaProfileEditorSettings::GetUserMediaProfile() const
 {
 	return UserMediaProfile.LoadSynchronous();
diff --git a/Engine/Plugins/Media/MediaFrameworkUtilities/Source/MediaFrameworkUtilities/Public/Profile/MediaProfile.h b/Engine/Plugins/Media/MediaFrameworkUtilities/Source/MediaFrameworkUtilities/Public/Profile/MediaProfile.h
--- a/Engine/Plugins/Media/MediaFrameworkUtilities/Source/MediaFrameworkUtilities/Public/Profile/MediaProfile.h
+++ b/Engine/Plugins/Media/MediaFrameworkUtilities/Source/MediaFrameworkUtilities/Public/Profile/MediaProfile.h
@@ -56,6 +56,11 @@ public:
 	 */
 	UMediaSource* GetMediaSource(int32 Index) const;
 
+	/**
+	 * Get the number of media source.
+	 */
+	int32 NumMediaSources() const;
+
 	/**
 	 * Get the media output for the selected proxy.
 	 *
@@ -63,6 +68,11 @@ public:
 	 */
 	UMediaOutput* GetMediaOutput(int32 Index) const;
 
+	/**
+	 * Get the number of media output.
+	 */
+	int32 NumMediaOutputs() const;
+
 	/**
 	 * Get the timecode provider.
 	 *
diff --git a/Engine/Plugins/Media/MediaFrameworkUtilities/Source/MediaFrameworkUtilities/Public/Profile/MediaProfileSettings.h b/Engine/Plugins/Media/MediaFrameworkUtilities/Source/MediaFrameworkUtilities/Public/Profile/MediaProfileSettings.h
--- a/Engine/Plugins/Media/MediaFrameworkUtilities/Source/MediaFrameworkUtilities/Public/Profile/MediaProfileSettings.h
+++ b/Engine/Plugins/Media/MediaFrameworkUtilities/Source/MediaFrameworkUtilities/Public/Profile/MediaProfileSettings.h
@@ -19,22 +19,14 @@ class UProxyMediaSource;
 /**
  * Settings for the media profile.
  */
-UCLASS(config=Game)
+UCLASS(config=Game, defaultconfig)
 class MEDIAFRAMEWORKUTILITIES_API UMediaProfileSettings
 	: public UObject
 {
-	GENERATED_BODY()
-
-	UMediaProfileSettings();		
+	GENERATED_BODY()	
 
 public:
 
-	/**
-	 * Display the media profile icon in the editor toolbar.
-	 */
-	UPROPERTY(Config, EditAnywhere, Category="MediaProfile")
-	bool bDisplayInToolbar;
-
 	/**
 	 * Apply the startup media profile even when we are running a commandlet.
 	 * @note We always try to apply the user media profile before the startup media profile in the editor or standalone.
@@ -91,6 +83,16 @@ class MEDIAFRAMEWORKUTILITIES_API UMediaProfileEditorSettings
 {
 	GENERATED_BODY()
 
+	UMediaProfileEditorSettings();
+
+public:
+
+	/**
+	 * Display the media profile icon in the editor toolbar.
+	 */
+	UPROPERTY(Config, EditAnywhere, Category = "MediaProfile")
+	bool bDisplayInToolbar;
+
 private:
 
 	/**
diff --git a/Engine/Plugins/Media/MediaFrameworkUtilities/Source/MediaFrameworkUtilitiesEditor/Private/UI/MediaProfileMenuEntry.cpp b/Engine/Plugins/Media/MediaFrameworkUtilities/Source/MediaFrameworkUtilitiesEditor/Private/UI/MediaProfileMenuEntry.cpp
--- a/Engine/Plugins/Media/MediaFrameworkUtilities/Source/MediaFrameworkUtilitiesEditor/Private/UI/MediaProfileMenuEntry.cpp
+++ b/Engine/Plugins/Media/MediaFrameworkUtilities/Source/MediaFrameworkUtilitiesEditor/Private/UI/MediaProfileMenuEntry.cpp
@@ -180,7 +180,7 @@ TUniquePtr<FMediaProfileMenuEntryImpl> FMediaProfileMenuEntryImpl::Implementatio
 
 void FMediaProfileMenuEntry::Register()
 {
-	if (!IsRunningCommandlet() && GetDefault<UMediaProfileSettings>()->bDisplayInToolbar)
+	if (!IsRunningCommandlet() && GetDefault<UMediaProfileEditorSettings>()->bDisplayInToolbar)
 	{
 		FMediaProfileMenuEntryImpl::Implementation = MakeUnique<FMediaProfileMenuEntryImpl>();
 	}
