Fix for MediaFrameworkUtilities plugin

diff --git a/Engine/Plugins/Media/MediaFrameworkUtilities/MediaFrameworkUtilities.uplugin b/Engine/Plugins/Media/MediaFrameworkUtilities/MediaFrameworkUtilities.uplugin
--- a/Engine/Plugins/Media/MediaFrameworkUtilities/MediaFrameworkUtilities.uplugin
+++ b/Engine/Plugins/Media/MediaFrameworkUtilities/MediaFrameworkUtilities.uplugin
@@ -38,6 +38,10 @@
 	],
 	"Plugins" :
 	[
+		{
+			"Name": "MediaIOFramework",
+			"Enabled": true
+		},
 		{
 			"Name" : "OpenCVLensDistortion",
 			"Enabled" : true
