Initialize property MediaIOFramework and MediaFrameworkUtilitites.

diff --git a/Engine/Plugins/Media/MediaFrameworkUtilities/Source/MediaFrameworkUtilitiesEditor/Private/MediaFrameworkWorldSettingsAssetUserData.cpp b/Engine/Plugins/Media/MediaFrameworkUtilities/Source/MediaFrameworkUtilitiesEditor/Private/MediaFrameworkWorldSettingsAssetUserData.cpp
--- a/Engine/Plugins/Media/MediaFrameworkUtilities/Source/MediaFrameworkUtilitiesEditor/Private/MediaFrameworkWorldSettingsAssetUserData.cpp
+++ b/Engine/Plugins/Media/MediaFrameworkUtilities/Source/MediaFrameworkUtilitiesEditor/Private/MediaFrameworkWorldSettingsAssetUserData.cpp
@@ -5,7 +5,8 @@
 
 
 FMediaFrameworkCaptureCameraViewportCameraOutputInfo::FMediaFrameworkCaptureCameraViewportCameraOutputInfo()
-	: ViewMode(VMI_Unknown)
+	: MediaOutput(nullptr)
+	, ViewMode(VMI_Unknown)
 {
 }
 
diff --git a/Engine/Plugins/Media/MediaIOFramework/Source/MediaIOCore/Private/MediaIOCoreDefinitions.cpp b/Engine/Plugins/Media/MediaIOFramework/Source/MediaIOCore/Private/MediaIOCoreDefinitions.cpp
--- a/Engine/Plugins/Media/MediaIOFramework/Source/MediaIOCore/Private/MediaIOCoreDefinitions.cpp
+++ b/Engine/Plugins/Media/MediaIOFramework/Source/MediaIOCore/Private/MediaIOCoreDefinitions.cpp
@@ -110,7 +110,10 @@ bool FMediaIOConnection::IsValid() const
  * FMediaIOMode
  */
 FMediaIOMode::FMediaIOMode()
-	: DeviceModeIdentifier(MediaIOCoreDefinitions::InvalidDeviceModeIdentifier)
+	: FrameRate(30, 1)
+	, Resolution(1920, 1080)
+	, Standard(EMediaIOStandardType::Progressive)
+	, DeviceModeIdentifier(MediaIOCoreDefinitions::InvalidDeviceModeIdentifier)
 {}
 
 
