No need to hide the UI for MediaIO customization.
Add apply button to the media profile and warn the user if the media profle was not applied properly.

diff --git a/Engine/Plugins/Media/MediaFrameworkUtilities/Source/MediaFrameworkUtilities/Private/Profile/MediaProfile.cpp b/Engine/Plugins/Media/MediaFrameworkUtilities/Source/MediaFrameworkUtilities/Private/Profile/MediaProfile.cpp
--- a/Engine/Plugins/Media/MediaFrameworkUtilities/Source/MediaFrameworkUtilities/Private/Profile/MediaProfile.cpp
+++ b/Engine/Plugins/Media/MediaFrameworkUtilities/Source/MediaFrameworkUtilities/Private/Profile/MediaProfile.cpp
@@ -61,6 +61,10 @@ UEngineCustomTimeStep* UMediaProfile::GetCustomTimeStep() const
 
 void UMediaProfile::Apply()
 {
+#if WITH_EDITORONLY_DATA
+	bNeedToBeReapplied = false;
+#endif
+
 	if (GEngine == nullptr)
 	{
 		UE_LOG(LogMediaFrameworkUtilities, Error, TEXT("The MediaProfile '%s' could not be applied. The Engine is not initialized."), *GetName());
@@ -186,13 +190,25 @@ void UMediaProfile::Reset()
 		}
 	}
 
-	if (bOverrideTimecodeProvider)
 	{
-		GEngine->SetTimecodeProvider(nullptr);
+		const UTimecodeProvider* CurrentTimecodeProvider = GEngine->GetTimecodeProvider();
+		if (CurrentTimecodeProvider)
+		{
+			if (CurrentTimecodeProvider->GetOuter() == this)
+			{
+				GEngine->SetTimecodeProvider(nullptr);
+			}
+		}
 	}
 
-	if (bOverrideCustomTimeStep)
 	{
-		GEngine->SetCustomTimeStep(GEngine->GetDefaultCustomTimeStep());
+		const UEngineCustomTimeStep* CurrentCustomTimeStep = GEngine->GetCustomTimeStep();
+		if (CurrentCustomTimeStep)
+		{
+			if (CurrentCustomTimeStep->GetOuter() == this)
+			{
+				GEngine->SetCustomTimeStep(GEngine->GetDefaultCustomTimeStep());
+			}
+		}
 	}
 }
diff --git a/Engine/Plugins/Media/MediaFrameworkUtilities/Source/MediaFrameworkUtilities/Public/Profile/MediaProfile.h b/Engine/Plugins/Media/MediaFrameworkUtilities/Source/MediaFrameworkUtilities/Public/Profile/MediaProfile.h
--- a/Engine/Plugins/Media/MediaFrameworkUtilities/Source/MediaFrameworkUtilities/Public/Profile/MediaProfile.h
+++ b/Engine/Plugins/Media/MediaFrameworkUtilities/Source/MediaFrameworkUtilities/Public/Profile/MediaProfile.h
@@ -47,6 +47,15 @@ private:
 	UPROPERTY(EditAnywhere, Instanced, Category="Genlock", meta=(EditCondition="bOverrideCustomTimeStep"))
 	UEngineCustomTimeStep* CustomTimeStep;
 
+public:
+#if WITH_EDITORONLY_DATA
+	/**
+	 * When the profile is the current profile and modifications made it dirty.
+	 * Without re-apply the profile does modification won't take effect.
+	 */
+	bool bNeedToBeReapplied;
+#endif
+
 public:
 
 	/**
diff --git a/Engine/Plugins/Media/MediaFrameworkUtilities/Source/MediaFrameworkUtilitiesEditor/Private/AssetEditor/MediaProfileCommands.cpp b/Engine/Plugins/Media/MediaFrameworkUtilities/Source/MediaFrameworkUtilitiesEditor/Private/AssetEditor/MediaProfileCommands.cpp
new file mode 100644
--- /dev/null
+++ b/Engine/Plugins/Media/MediaFrameworkUtilities/Source/MediaFrameworkUtilitiesEditor/Private/AssetEditor/MediaProfileCommands.cpp
@@ -0,0 +1,12 @@
+// Copyright 1998-2018 Epic Games, Inc. All Rights Reserved.
+
+#include "MediaProfileCommands.h"
+
+#define LOCTEXT_NAMESPACE "MediaProfileCommands"
+
+void FMediaProfileCommands::RegisterCommands()
+{
+	UI_COMMAND(Apply, "Apply", "Apply changes to the media profile.", EUserInterfaceActionType::Button, FInputChord());
+}
+
+#undef LOCTEXT_NAMESPACE
\ No newline at end of file
diff --git a/Engine/Plugins/Media/MediaFrameworkUtilities/Source/MediaFrameworkUtilitiesEditor/Private/AssetEditor/MediaProfileCommands.h b/Engine/Plugins/Media/MediaFrameworkUtilities/Source/MediaFrameworkUtilitiesEditor/Private/AssetEditor/MediaProfileCommands.h
new file mode 100644
--- /dev/null
+++ b/Engine/Plugins/Media/MediaFrameworkUtilities/Source/MediaFrameworkUtilitiesEditor/Private/AssetEditor/MediaProfileCommands.h
@@ -0,0 +1,25 @@
+// Copyright 1998-2018 Epic Games, Inc. All Rights Reserved.
+
+#pragma once
+
+#include "CoreMinimal.h"
+#include "EditorStyleSet.h"
+#include "Framework/Commands/Commands.h"
+#include "UI/MediaFrameworkUtilitiesEditorStyle.h"
+
+class FMediaProfileCommands : public TCommands<FMediaProfileCommands>
+{
+public:
+
+	FMediaProfileCommands()
+		: TCommands<FMediaProfileCommands>(TEXT("ToolbarIcon"), NSLOCTEXT("Contexts", "ToolbarIcon", "Media Profile"), NAME_None, FMediaFrameworkUtilitiesEditorStyle::GetStyleSetName())
+	{}
+	
+	/**
+	 * Initialize commands
+	 */
+	virtual void RegisterCommands() override;
+
+	/** Applies changes to the original Media Profile.*/
+	TSharedPtr<FUICommandInfo> Apply;
+};
diff --git a/Engine/Plugins/Media/MediaFrameworkUtilities/Source/MediaFrameworkUtilitiesEditor/Private/AssetEditor/MediaProfileEditorToolkit.cpp b/Engine/Plugins/Media/MediaFrameworkUtilities/Source/MediaFrameworkUtilitiesEditor/Private/AssetEditor/MediaProfileEditorToolkit.cpp
new file mode 100644
--- /dev/null
+++ b/Engine/Plugins/Media/MediaFrameworkUtilities/Source/MediaFrameworkUtilitiesEditor/Private/AssetEditor/MediaProfileEditorToolkit.cpp
@@ -0,0 +1,168 @@
+// Copyright 1998-2018 Epic Games, Inc. All Rights Reserved.
+
+#include "AssetEditor/MediaProfileEditorToolkit.h"
+#include "Profile/MediaProfile.h"
+#include "Framework/MultiBox/MultiBoxBuilder.h"
+#include "Framework/MultiBox/MultiBoxExtender.h"
+#include "UI/MediaFrameworkUtilitiesEditorStyle.h"
+#include "MediaProfileCommands.h"
+#include "Misc/MessageDialog.h"
+#include "Profile/MediaProfileSettings.h"
+#include "Profile/IMediaProfileManager.h"
+
+#define LOCTEXT_NAMESPACE "MediaProfileEditor"
+
+TSharedRef<FMediaProfileEditorToolkit> FMediaProfileEditorToolkit::CreateEditor(const EToolkitMode::Type Mode, const TSharedPtr< class IToolkitHost >& InitToolkitHost, class UMediaProfile* InMediaProfile)
+{
+	TSharedRef<FMediaProfileEditorToolkit> NewEditor = MakeShared<FMediaProfileEditorToolkit>();
+
+	NewEditor->InitMediaProfileEditor(Mode, InitToolkitHost, InMediaProfile);
+
+	return NewEditor;
+}
+
+void FMediaProfileEditorToolkit::InitMediaProfileEditor(const EToolkitMode::Type Mode, const TSharedPtr< class IToolkitHost >& InitToolkitHost, UMediaProfile* InMediaProfile)
+{
+	TArray<UObject*> ObjectsToEdit;
+	ObjectsToEdit.Add(InMediaProfile);
+
+	Super::InitEditor(Mode, InitToolkitHost, ObjectsToEdit, FGetDetailsViewObjects());
+
+	bSubPropertyWasModified = false;
+
+	BindCommands();
+
+	ExtendToolBar();
+
+	FCoreUObjectDelegates::OnObjectPropertyChanged.AddSP(this, &FMediaProfileEditorToolkit::HandleCoreObjectPropertyChanged);
+	FCoreUObjectDelegates::OnPreObjectPropertyChanged.AddSP(this, &FMediaProfileEditorToolkit::HandleCorePreObjectPropertyChanged);
+}
+
+FMediaProfileEditorToolkit::~FMediaProfileEditorToolkit()
+{
+	FCoreUObjectDelegates::OnPreObjectPropertyChanged.RemoveAll(this);
+	FCoreUObjectDelegates::OnObjectPropertyChanged.RemoveAll(this);
+}
+
+void FMediaProfileEditorToolkit::ExtendToolBar()
+{
+	TSharedRef<FExtender> ToolbarExtender = MakeShared<FExtender>();
+
+	ToolbarExtender->AddToolBarExtension(
+		"Asset",
+		EExtensionHook::After,
+		GetToolkitCommands(),
+		FToolBarExtensionDelegate::CreateStatic([](FToolBarBuilder& ToolbarBuilder)
+		{
+			ToolbarBuilder.BeginSection("Apply");
+			{
+				ToolbarBuilder.AddToolBarButton(FMediaProfileCommands::Get().Apply);
+			}
+		})
+	);
+
+	AddToolbarExtender(ToolbarExtender);
+
+	RegenerateMenusAndToolbars();
+}
+
+void FMediaProfileEditorToolkit::BindCommands()
+{
+	ToolkitCommands->MapAction(
+		FMediaProfileCommands::Get().Apply,
+		FExecuteAction::CreateSP(this, &FMediaProfileEditorToolkit::ApplyMediaProfile),
+		FCanExecuteAction::CreateLambda([this]()
+			{
+				if (UMediaProfile* MediaProfile = Cast<UMediaProfile>(GetEditingObject()))
+				{
+					return MediaProfile->bNeedToBeReapplied;
+				}
+				return false;
+			})
+	);
+}
+
+void FMediaProfileEditorToolkit::ApplyMediaProfile()
+{
+	UMediaProfile* MediaProfile = Cast<UMediaProfile>(GetEditingObject());
+	IMediaProfileManager::Get().SetCurrentMediaProfile(nullptr);
+
+	if (MediaProfile != nullptr)
+	{
+		IMediaProfileManager::Get().SetCurrentMediaProfile(MediaProfile);
+	}
+}
+
+void FMediaProfileEditorToolkit::HandleCorePreObjectPropertyChanged(UObject* Object, const FEditPropertyChain& EditPropertyChain)
+{
+	if (Object == GetEditingObject())
+	{
+		bSubPropertyWasModified = false;
+	}
+}
+
+
+void FMediaProfileEditorToolkit::HandleCoreObjectPropertyChanged(UObject* Object, struct FPropertyChangedEvent& ChangedEvent)
+{
+	if (Object == nullptr)
+	{
+		return;
+	}
+
+	if (Object == GetEditingObject())
+	{
+		// We don't want to display apply if a sub property was modified (ie. FilePath in a FileMediaSource),
+		// So only show the apply button if the ObjectPropertyChanged was fired for a direct property change on MediaProfile.
+		if (!bSubPropertyWasModified)
+		{
+			if (UMediaProfile* MediaProfile = Cast<UMediaProfile>(GetEditingObject()))
+			{
+				MediaProfile->bNeedToBeReapplied = true;
+			}
+		}
+
+		bSubPropertyWasModified = false;
+	}
+	else if (Object->GetOuter() == GetEditingObject())
+	{
+		bSubPropertyWasModified = true;
+	}
+}
+
+void FMediaProfileEditorToolkit::SaveAsset_Execute()
+{
+	ApplyMediaProfile();
+	Super::SaveAsset_Execute();
+}
+
+void FMediaProfileEditorToolkit::SaveAssetAs_Execute()
+{
+	ApplyMediaProfile();
+	Super::SaveAssetAs_Execute();
+}
+
+bool FMediaProfileEditorToolkit::OnRequestClose()
+{
+	bool bClose = true;
+	if (UMediaProfile* MediaProfile = Cast<UMediaProfile>(GetEditingObject()))
+	{
+		if (MediaProfile->bNeedToBeReapplied)
+		{
+			// find out the user wants to do with this dirty profile
+			EAppReturnType::Type YesNoCancelReply = FMessageDialog::Open(EAppMsgType::YesNoCancel, LOCTEXT("Prompt_EditorClose", "Would you like to apply the modifications?"));
+
+			switch (YesNoCancelReply)
+			{
+			case EAppReturnType::Yes:
+				ApplyMediaProfile();
+				break;
+			case EAppReturnType::Cancel:
+				bClose = false;
+				break;
+			}
+		}
+	}
+	return bClose;
+}
+
+#undef LOCTEXT_NAMESPACE
\ No newline at end of file
diff --git a/Engine/Plugins/Media/MediaFrameworkUtilities/Source/MediaFrameworkUtilitiesEditor/Private/AssetEditor/MediaProfileEditorToolkit.h b/Engine/Plugins/Media/MediaFrameworkUtilities/Source/MediaFrameworkUtilitiesEditor/Private/AssetEditor/MediaProfileEditorToolkit.h
new file mode 100644
--- /dev/null
+++ b/Engine/Plugins/Media/MediaFrameworkUtilities/Source/MediaFrameworkUtilitiesEditor/Private/AssetEditor/MediaProfileEditorToolkit.h
@@ -0,0 +1,68 @@
+// Copyright 1998-2018 Epic Games, Inc. All Rights Reserved.
+
+#pragma once
+
+#include "CoreMinimal.h"
+
+#include "Factories/Factory.h"
+#include "Profile/MediaProfile.h"
+#include "Toolkits/SimpleAssetEditor.h"
+
+class FExtender;
+
+/** Viewer/editor for a MediaProfile */
+class FMediaProfileEditorToolkit : public FSimpleAssetEditor
+{
+
+private:
+	using Super = FSimpleAssetEditor;
+
+public:
+	static TSharedRef<FMediaProfileEditorToolkit> CreateEditor(const EToolkitMode::Type Mode, const TSharedPtr< class IToolkitHost >& InitToolkitHost, UMediaProfile* InMediaProfile);
+
+	virtual ~FMediaProfileEditorToolkit();
+
+	/**
+	 * Edits the specified table
+	 *
+	 * @param	Mode					Asset editing mode for this editor (standalone or world-centric)
+	 * @param	InitToolkitHost			When Mode is WorldCentric, this is the level editor instance to spawn this editor within
+	 * @param	InMediaProfile			The MediaProfile asset to edit
+	 */
+	void InitMediaProfileEditor(const EToolkitMode::Type Mode, const TSharedPtr< class IToolkitHost >& InitToolkitHost, UMediaProfile* InMediaProfile);
+
+	/**
+	 * Binds our UI commands to delegates
+	 */
+	void BindCommands();
+
+protected:
+
+	virtual void SaveAsset_Execute() override;
+	virtual void SaveAssetAs_Execute() override;
+	virtual bool OnRequestClose() override;
+
+private:
+
+	/** Builds the toolbar widget for media profile editor */
+	void ExtendToolBar();
+
+	/** Applies changes to the media profile */
+	void ApplyMediaProfile();
+
+	/** Handle object property changes */
+	void HandleCoreObjectPropertyChanged(UObject* Object, struct FPropertyChangedEvent& ChangedEvent);
+
+	/** Handle pre object property changes */
+	void HandleCorePreObjectPropertyChanged(UObject* Object, const class FEditPropertyChain& EditPropertyChain);
+
+private:
+	/** Dockable tab for properties */
+	TSharedPtr<class SDockableTab> PropertiesTab;
+
+	/** Details view */
+	TSharedPtr<class IDetailsView> DetailsView;
+	
+	/** Whether the last modified object was a sub property of MediaProfile */
+	bool bSubPropertyWasModified;
+};
\ No newline at end of file
diff --git a/Engine/Plugins/Media/MediaFrameworkUtilities/Source/MediaFrameworkUtilitiesEditor/Private/AssetTypeActions/AssetTypeActions_MediaProfile .cpp b/Engine/Plugins/Media/MediaFrameworkUtilities/Source/MediaFrameworkUtilitiesEditor/Private/AssetTypeActions/AssetTypeActions_MediaProfile .cpp
new file mode 100644
--- /dev/null
+++ b/Engine/Plugins/Media/MediaFrameworkUtilities/Source/MediaFrameworkUtilitiesEditor/Private/AssetTypeActions/AssetTypeActions_MediaProfile .cpp	
@@ -0,0 +1,31 @@
+// Copyright 1998-2018 Epic Games, Inc. All Rights Reserved.
+
+#include "AssetTypeActions/AssetTypeActions_MediaProfile.h"
+
+#include "AssetEditor/MediaProfileEditorToolkit.h"
+#include "Profile/MediaProfile.h"
+
+#define LOCTEXT_NAMESPACE "MediaProfileEditor"
+
+FText FAssetTypeActions_MediaProfile::GetName() const
+{
+	return LOCTEXT("AssetTypeActions_MediaProfile", "Media Profile");
+}
+
+UClass* FAssetTypeActions_MediaProfile::GetSupportedClass() const
+{
+	return UMediaProfile::StaticClass();
+}
+
+void FAssetTypeActions_MediaProfile::OpenAssetEditor(const TArray<UObject*>& InObjects, TSharedPtr<class IToolkitHost> EditWithinLevelEditor)
+{
+	for (UObject* Obj : InObjects)
+	{
+		if (UMediaProfile* Asset = Cast<UMediaProfile>(Obj))
+		{
+			FMediaProfileEditorToolkit::CreateEditor(EToolkitMode::Standalone, EditWithinLevelEditor, Asset);
+		}
+	}
+}
+
+#undef LOCTEXT_NAMESPACE
diff --git a/Engine/Plugins/Media/MediaFrameworkUtilities/Source/MediaFrameworkUtilitiesEditor/Private/AssetTypeActions/AssetTypeActions_MediaProfile.h b/Engine/Plugins/Media/MediaFrameworkUtilities/Source/MediaFrameworkUtilitiesEditor/Private/AssetTypeActions/AssetTypeActions_MediaProfile.h
new file mode 100644
--- /dev/null
+++ b/Engine/Plugins/Media/MediaFrameworkUtilities/Source/MediaFrameworkUtilitiesEditor/Private/AssetTypeActions/AssetTypeActions_MediaProfile.h
@@ -0,0 +1,19 @@
+// Copyright 1998-2018 Epic Games, Inc. All Rights Reserved.
+
+#pragma once
+
+#include "CoreMinimal.h"
+#include "AssetTypeActions_Base.h"
+
+class FAssetTypeActions_MediaProfile : public FAssetTypeActions_Base
+{
+public:
+	// IAssetTypeActions Implementation
+	virtual FColor GetTypeColor() const override { return FColor(140, 62, 35); }
+	virtual bool HasActions(const TArray<UObject*>& InObjects) const override { return false; }
+	virtual uint32 GetCategories() override { return EAssetTypeCategories::Media; }
+	virtual bool IsImportedAsset() const override { return false; }
+	virtual FText GetName() const override;
+	virtual UClass* GetSupportedClass() const override;
+	virtual void OpenAssetEditor(const TArray<UObject*>& InObjects, TSharedPtr<class IToolkitHost> EditWithinLevelEditor = TSharedPtr<IToolkitHost>()) override;
+};
\ No newline at end of file
diff --git a/Engine/Plugins/Media/MediaFrameworkUtilities/Source/MediaFrameworkUtilitiesEditor/Private/MediaFrameworkUtilitiesEditorModule.cpp b/Engine/Plugins/Media/MediaFrameworkUtilities/Source/MediaFrameworkUtilitiesEditor/Private/MediaFrameworkUtilitiesEditorModule.cpp
--- a/Engine/Plugins/Media/MediaFrameworkUtilities/Source/MediaFrameworkUtilitiesEditor/Private/MediaFrameworkUtilitiesEditorModule.cpp
+++ b/Engine/Plugins/Media/MediaFrameworkUtilities/Source/MediaFrameworkUtilitiesEditor/Private/MediaFrameworkUtilitiesEditorModule.cpp
@@ -11,12 +11,14 @@
 #include "UObject/UObjectBase.h"
 
 #include "AssetTypeActions/AssetTypeActions_MediaBundle.h"
+#include "AssetTypeActions/AssetTypeActions_MediaProfile.h"
 #include "MediaBundleActorDetails.h"
 #include "MediaBundleFactoryNew.h"
 #include "MediaFrameworkUtilitiesPlacement.h"
 #include "CaptureTab/SMediaFrameworkCapture.h"
 #include "UI/MediaFrameworkUtilitiesEditorStyle.h"
 #include "UI/MediaProfileMenuEntry.h"
+#include "AssetEditor/MediaProfileCommands.h"
 
 #define LOCTEXT_NAMESPACE "MediaFrameworkEditor"
 
@@ -40,14 +42,22 @@ public:
 			FMediaFrameworkUtilitiesPlacement::RegisterPlacement();
 
 			// Register AssetTypeActions
-			AssetTypeAction = MakeShareable(new FAssetTypeActions_MediaBundle());
-			FModuleManager::LoadModuleChecked<FAssetToolsModule>("AssetTools").Get().RegisterAssetTypeActions(AssetTypeAction.ToSharedRef());
+			auto RegisterAssetTypeAction = [this](const TSharedRef<IAssetTypeActions>& InAssetTypeAction)
+			{
+				IAssetTools& AssetTools = FModuleManager::LoadModuleChecked<FAssetToolsModule>("AssetTools").Get();
+				RegisteredAssetTypeActions.Add(InAssetTypeAction);
+				AssetTools.RegisterAssetTypeActions(InAssetTypeAction);
+			};
+
+			RegisterAssetTypeAction(MakeShared<FAssetTypeActions_MediaBundle>());
+			RegisterAssetTypeAction(MakeShared<FAssetTypeActions_MediaProfile>());
 
 			FPropertyEditorModule& PropertyModule = FModuleManager::GetModuleChecked<FPropertyEditorModule>("PropertyEditor");
 			PropertyModule.RegisterCustomClassLayout("MediaBundleActorBase", FOnGetDetailCustomizationInstance::CreateStatic(&FMediaBundleActorDetails::MakeInstance));
 
 			SMediaFrameworkCapture::RegisterNomadTabSpawner();
 			FMediaProfileMenuEntry::Register();
+			FMediaProfileCommands::Register();
 		}
 	}
 
@@ -55,14 +65,25 @@ public:
 	{
 		if (!GIsRequestingExit && GEditor && UObjectInitialized())
 		{
-			FMediaProfileMenuEntry::Register();
+			FMediaProfileCommands::Unregister();
+			FMediaProfileMenuEntry::Unregister();
 			SMediaFrameworkCapture::UnregisterNomadTabSpawner();
 
 			FPropertyEditorModule& PropertyModule = FModuleManager::GetModuleChecked<FPropertyEditorModule>("PropertyEditor");
 			PropertyModule.UnregisterCustomClassLayout("MediaBundleActorBase");
 
 			// Unregister AssetTypeActions
-			FModuleManager::LoadModuleChecked<FAssetToolsModule>("AssetTools").Get().UnregisterAssetTypeActions(AssetTypeAction.ToSharedRef());
+			FAssetToolsModule* AssetToolsModule = FModuleManager::GetModulePtr<FAssetToolsModule>("AssetTools");
+
+			if (AssetToolsModule != nullptr)
+			{
+				IAssetTools& AssetTools = AssetToolsModule->Get();
+
+				for (auto Action : RegisteredAssetTypeActions)
+				{
+					AssetTools.UnregisterAssetTypeActions(Action);
+				}
+			}
 
 			FMediaFrameworkUtilitiesPlacement::UnregisterPlacement();
 
@@ -72,7 +93,9 @@ public:
 		}
 	}
 
-	TSharedPtr<FAssetTypeActions_MediaBundle> AssetTypeAction;
+private:
+
+	TArray<TSharedRef<IAssetTypeActions>> RegisteredAssetTypeActions;
 };
 
 
diff --git a/Engine/Plugins/Media/MediaFrameworkUtilities/Source/MediaFrameworkUtilitiesEditor/Private/UI/MediaFrameworkUtilitiesEditorStyle.cpp b/Engine/Plugins/Media/MediaFrameworkUtilities/Source/MediaFrameworkUtilitiesEditor/Private/UI/MediaFrameworkUtilitiesEditorStyle.cpp
--- a/Engine/Plugins/Media/MediaFrameworkUtilities/Source/MediaFrameworkUtilitiesEditor/Private/UI/MediaFrameworkUtilitiesEditorStyle.cpp
+++ b/Engine/Plugins/Media/MediaFrameworkUtilities/Source/MediaFrameworkUtilitiesEditor/Private/UI/MediaFrameworkUtilitiesEditorStyle.cpp
@@ -35,6 +35,7 @@ void FMediaFrameworkUtilitiesEditorStyle::Register()
 	MediaFrameworkUtilitiesStyle::StyleInstance->Set("ClassThumbnail.MediaProfile", new IMAGE_BRUSH("MediaProfile_64x", MediaFrameworkUtilitiesStyle::Icon64x64));
 	MediaFrameworkUtilitiesStyle::StyleInstance->Set("ClassIcon.MediaProfile", new IMAGE_BRUSH("MediaProfile_20x", MediaFrameworkUtilitiesStyle::Icon20x20));
 	MediaFrameworkUtilitiesStyle::StyleInstance->Set("ToolbarIcon.MediaProfile", new IMAGE_BRUSH("MediaProfile_Color_40x", MediaFrameworkUtilitiesStyle::Icon40x40));
+	MediaFrameworkUtilitiesStyle::StyleInstance->Set("ToolbarIcon.Apply", new IMAGE_BRUSH("Apply_40x", MediaFrameworkUtilitiesStyle::Icon40x40));
 
 	MediaFrameworkUtilitiesStyle::StyleInstance->Set("CaptureCameraViewport_Capture", new IMAGE_BRUSH("CaptureCameraViewport_Capture_40x", MediaFrameworkUtilitiesStyle::Icon40x40));
 	MediaFrameworkUtilitiesStyle::StyleInstance->Set("CaptureCameraViewport_Capture.Small", new IMAGE_BRUSH("CaptureCameraViewport_Capture_16x", MediaFrameworkUtilitiesStyle::Icon16x16));
diff --git a/Engine/Plugins/Media/MediaIOFramework/Source/MediaIOEditor/Private/Customizations/MediaIOCustomizationBase.cpp b/Engine/Plugins/Media/MediaIOFramework/Source/MediaIOEditor/Private/Customizations/MediaIOCustomizationBase.cpp
--- a/Engine/Plugins/Media/MediaIOFramework/Source/MediaIOEditor/Private/Customizations/MediaIOCustomizationBase.cpp
+++ b/Engine/Plugins/Media/MediaIOFramework/Source/MediaIOEditor/Private/Customizations/MediaIOCustomizationBase.cpp
@@ -114,7 +114,6 @@ void FMediaIOCustomizationBase::CustomizeChildren(TSharedRef<IPropertyHandle> In
 			StructBuilder.AddProperty(InStructPropertyHandle->GetChildHandle(Index).ToSharedRef()).IsEnabled(false).ShowPropertyButtons(false);
 		}
 	}
-	InStructPropertyHandle->MarkHiddenByCustomization();
 }
 
 void FMediaIOCustomizationBase::AssignValueImpl(const void* NewValue) const
