Improve  media player & timecodesynchronizer toolbar UX. Separate the main button from the combo box.

diff --git a/Engine/Plugins/Media/MediaFrameworkUtilities/Source/MediaFrameworkUtilities/Public/Profile/MediaProfileSettings.h b/Engine/Plugins/Media/MediaFrameworkUtilities/Source/MediaFrameworkUtilities/Public/Profile/MediaProfileSettings.h
--- a/Engine/Plugins/Media/MediaFrameworkUtilities/Source/MediaFrameworkUtilities/Public/Profile/MediaProfileSettings.h
+++ b/Engine/Plugins/Media/MediaFrameworkUtilities/Source/MediaFrameworkUtilities/Public/Profile/MediaProfileSettings.h
@@ -90,7 +90,7 @@ public:
 	/**
 	 * Display the media profile icon in the editor toolbar.
 	 */
-	UPROPERTY(Config, EditAnywhere, Category = "MediaProfile")
+	UPROPERTY(Config, EditAnywhere, Category = "MediaProfile", meta=(ConfigRestartRequired=true))
 	bool bDisplayInToolbar;
 
 private:
diff --git a/Engine/Plugins/Media/MediaFrameworkUtilities/Source/MediaFrameworkUtilitiesEditor/Private/AssetEditor/MediaProfileCommands.cpp b/Engine/Plugins/Media/MediaFrameworkUtilities/Source/MediaFrameworkUtilitiesEditor/Private/AssetEditor/MediaProfileCommands.cpp
--- a/Engine/Plugins/Media/MediaFrameworkUtilities/Source/MediaFrameworkUtilitiesEditor/Private/AssetEditor/MediaProfileCommands.cpp
+++ b/Engine/Plugins/Media/MediaFrameworkUtilities/Source/MediaFrameworkUtilitiesEditor/Private/AssetEditor/MediaProfileCommands.cpp
@@ -2,11 +2,19 @@
 
 #include "MediaProfileCommands.h"
 
+#include "Framework/Commands/UICommandInfo.h"
+#include "Framework/Commands/UICommandList.h"
+#include "Profile/IMediaProfileManager.h"
+#include "Profile/MediaProfile.h"
+#include "Toolkits/AssetEditorManager.h"
+
+
 #define LOCTEXT_NAMESPACE "MediaProfileCommands"
 
 void FMediaProfileCommands::RegisterCommands()
 {
 	UI_COMMAND(Apply, "Apply", "Apply changes to the media profile.", EUserInterfaceActionType::Button, FInputChord());
+	UI_COMMAND(Edit, "Edit", "Edit the current media profile.", EUserInterfaceActionType::Button, FInputChord());
 }
 
 #undef LOCTEXT_NAMESPACE
\ No newline at end of file
diff --git a/Engine/Plugins/Media/MediaFrameworkUtilities/Source/MediaFrameworkUtilitiesEditor/Private/AssetEditor/MediaProfileCommands.h b/Engine/Plugins/Media/MediaFrameworkUtilities/Source/MediaFrameworkUtilitiesEditor/Private/AssetEditor/MediaProfileCommands.h
--- a/Engine/Plugins/Media/MediaFrameworkUtilities/Source/MediaFrameworkUtilitiesEditor/Private/AssetEditor/MediaProfileCommands.h
+++ b/Engine/Plugins/Media/MediaFrameworkUtilities/Source/MediaFrameworkUtilitiesEditor/Private/AssetEditor/MediaProfileCommands.h
@@ -7,6 +7,10 @@
 #include "Framework/Commands/Commands.h"
 #include "UI/MediaFrameworkUtilitiesEditorStyle.h"
 
+
+class FUICommandInfo;
+
+
 class FMediaProfileCommands : public TCommands<FMediaProfileCommands>
 {
 public:
@@ -20,6 +24,9 @@ public:
 	 */
 	virtual void RegisterCommands() override;
 
-	/** Applies changes to the original Media Profile.*/
+	/** Applies changes to the original Media Profile. */
 	TSharedPtr<FUICommandInfo> Apply;
+
+	/** Edit the current Media Profile. */
+	TSharedPtr<FUICommandInfo> Edit;
 };
diff --git a/Engine/Plugins/Media/MediaFrameworkUtilities/Source/MediaFrameworkUtilitiesEditor/Private/UI/MediaProfileMenuEntry.cpp b/Engine/Plugins/Media/MediaFrameworkUtilities/Source/MediaFrameworkUtilitiesEditor/Private/UI/MediaProfileMenuEntry.cpp
--- a/Engine/Plugins/Media/MediaFrameworkUtilities/Source/MediaFrameworkUtilitiesEditor/Private/UI/MediaProfileMenuEntry.cpp
+++ b/Engine/Plugins/Media/MediaFrameworkUtilities/Source/MediaFrameworkUtilitiesEditor/Private/UI/MediaProfileMenuEntry.cpp
@@ -2,6 +2,7 @@
 
 #include "UI/MediaProfileMenuEntry.h"
 
+#include "AssetEditor/MediaProfileCommands.h"
 #include "AssetToolsModule.h"
 #include "Factories/MediaProfileFactoryNew.h"
 #include "Framework/Application/SlateApplication.h"
@@ -19,15 +20,29 @@
 #include "Toolkits/AssetEditorManager.h"
 #include "UI/MediaFrameworkUtilitiesEditorStyle.h"
 
-
 #define LOCTEXT_NAMESPACE "MediaProfileEditor"
 
 struct FMediaProfileMenuEntryImpl
 {
 	FMediaProfileMenuEntryImpl()
 	{
+		TSharedPtr<FUICommandList> Actions = MakeShareable(new FUICommandList);
+
+		// Action to edit the current selected media profile
+		Actions->MapAction(FMediaProfileCommands::Get().Edit,
+			FExecuteAction::CreateLambda([this]()
+			{
+				if (UMediaProfile* MediaProfile = IMediaProfileManager::Get().GetCurrentMediaProfile())
+				{
+					FAssetEditorManager::Get().OpenEditorForAsset(MediaProfile);
+				}
+			}),
+			FCanExecuteAction::CreateLambda([] { return IMediaProfileManager::Get().GetCurrentMediaProfile() != nullptr; }),
+			FIsActionChecked::CreateLambda([]{ return IMediaProfileManager::Get().GetCurrentMediaProfile() != nullptr; })
+		);
+
 		ToolBarExtender = MakeShareable(new FExtender);
-		ToolBarExtender->AddToolBarExtension("Settings", EExtensionHook::After, nullptr, FToolBarExtensionDelegate::CreateRaw(this, &FMediaProfileMenuEntryImpl::FillToolbar));
+		ToolBarExtender->AddToolBarExtension("Settings", EExtensionHook::After, Actions, FToolBarExtensionDelegate::CreateRaw(this, &FMediaProfileMenuEntryImpl::FillToolbar));
 
 		FLevelEditorModule& LevelEditorModule = FModuleManager::LoadModuleChecked<FLevelEditorModule>("LevelEditor");
 		LevelEditorModule.GetToolBarExtensibilityManager()->AddExtender(ToolBarExtender);
@@ -50,11 +65,6 @@ struct FMediaProfileMenuEntryImpl
 		return IMediaProfileManager::Get().GetCurrentMediaProfile();
 	}
 
-	void OpenCurrentProfile()
-	{
-		FAssetEditorManager::Get().OpenEditorForAsset(GetCurrentProfile());
-	}
-
 	void CreateNewProfile()
 	{
 		UMediaProfileFactoryNew* FactoryInstance = DuplicateObject<UMediaProfileFactoryNew>(GetDefault<UMediaProfileFactoryNew>(), GetTransientPackage());
@@ -84,12 +94,34 @@ struct FMediaProfileMenuEntryImpl
 	{
 		ToolbarBuilder.BeginSection("Media Profile");
 		{
+			auto TooltipLambda = [this]()
+			{
+				UMediaProfile* MediaProfile = IMediaProfileManager::Get().GetCurrentMediaProfile();
+				if (MediaProfile == nullptr)
+				{
+					return LOCTEXT("EmptyMediaProfile_ToolTip", "Select a Media Profile to edit it.");
+				}
+				return FText::Format(LOCTEXT("MediaProfile_ToolTip", "Edit '{0}'")
+					, FText::FromName(MediaProfile->GetFName()));
+			};
+
+			// Add a button to edit the current media profile
+			ToolbarBuilder.AddToolBarButton(
+				FMediaProfileCommands::Get().Edit,
+				NAME_None,
+				LOCTEXT("MediaProfile_Label", "Media Profile"),
+				MakeAttributeLambda(TooltipLambda),
+				FSlateIcon(FMediaFrameworkUtilitiesEditorStyle::GetStyleSetName(), TEXT("ToolbarIcon.MediaProfile"))
+			);
+
+			// Add a simple drop-down menu (no label, no icon for the drop-down button itself) that list the media profile available
 			ToolbarBuilder.AddComboButton(
 				FUIAction(),
 				FOnGetContent::CreateRaw(this, &FMediaProfileMenuEntryImpl::GenerateMenuContent),
-				LOCTEXT("MediaProfile_Label", "Media Profile"),
-				LOCTEXT("MediaProfile_ToolTip", "List of Media Profile available to the user for editing or creation."),
-				FSlateIcon(FMediaFrameworkUtilitiesEditorStyle::GetStyleSetName(), TEXT("ToolbarIcon.MediaProfile"))
+				FText::GetEmpty(),
+				LOCTEXT("MediaProfileButton_ToolTip", "List of Media Profile available to the user for editing or creation."),
+				FSlateIcon(),
+				true
 			);
 		}
 		ToolbarBuilder.EndSection();
@@ -100,7 +132,7 @@ struct FMediaProfileMenuEntryImpl
 		const bool bShouldCloseWindowAfterMenuSelection = true;
 		FMenuBuilder MenuBuilder(bShouldCloseWindowAfterMenuSelection, nullptr);
 
-		MenuBuilder.BeginSection("Profile", LOCTEXT("MediaProfileSection", "Media Profile"));
+		MenuBuilder.BeginSection("Profile", LOCTEXT("NewMediaProfileSection", "New"));
 		{
 			MenuBuilder.AddMenuEntry(
 				LOCTEXT("CreateMenuLabel", "New Empty Media Profile"),
@@ -110,21 +142,19 @@ struct FMediaProfileMenuEntryImpl
 					FExecuteAction::CreateRaw(this, &FMediaProfileMenuEntryImpl::CreateNewProfile)
 				)
 			);
+		}
+		MenuBuilder.EndSection();
 
-			MenuBuilder.AddMenuSeparator();
-
+		MenuBuilder.BeginSection("Profile", LOCTEXT("MediaProfileSection", "Media Profile"));
+		{
 			UMediaProfile* Profile = GetCurrentProfile();
 			const bool bIsProfileValid = Profile != nullptr;
 
 			MenuBuilder.AddSubMenu(
-				bIsProfileValid ? FText::Format(LOCTEXT("EditMenuLabel", "Open '{0}'"), FText::FromName(Profile->GetFName()))  : LOCTEXT("SelectMenuLabel", "Select Profile"),
+				bIsProfileValid ? FText::FromName(Profile->GetFName()) : LOCTEXT("SelectMenuLabel", "Select a Media Profile"),
 				LOCTEXT("SelectMenuTooltip", "Select the current profile for this editor."),
 				FNewMenuDelegate::CreateRaw(this, &FMediaProfileMenuEntryImpl::AddObjectSubMenu),
-				FUIAction(
-					bIsProfileValid ? FExecuteAction::CreateRaw(this, &FMediaProfileMenuEntryImpl::OpenCurrentProfile) : FExecuteAction(),
-					FCanExecuteAction(),
-					FIsActionChecked::CreateLambda([bIsProfileValid]{ return bIsProfileValid; })
-				),
+				FUIAction(),
 				NAME_None,
 				EUserInterfaceActionType::RadioButton
 			);
diff --git a/Engine/Plugins/Media/TimecodeSynchronizer/Source/TimecodeSynchronizer/Public/TimecodeSynchronizerProjectSettings.h b/Engine/Plugins/Media/TimecodeSynchronizer/Source/TimecodeSynchronizer/Public/TimecodeSynchronizerProjectSettings.h
--- a/Engine/Plugins/Media/TimecodeSynchronizer/Source/TimecodeSynchronizer/Public/TimecodeSynchronizerProjectSettings.h
+++ b/Engine/Plugins/Media/TimecodeSynchronizer/Source/TimecodeSynchronizer/Public/TimecodeSynchronizerProjectSettings.h
@@ -25,7 +25,7 @@ public:
 
 public:
 	/** Display the timecode synchronizer icon in the editor toolbar. */
-	UPROPERTY(Config, EditAnywhere, Category="TimecodeSynchronizer")
+	UPROPERTY(Config, EditAnywhere, Category="TimecodeSynchronizer", meta=(ConfigRestartRequired=true))
 	bool bDisplayInToolbar;
 
 	UPROPERTY(config, EditAnywhere, Category="TimecodeSynchronizer")
diff --git a/Engine/Plugins/Media/TimecodeSynchronizer/Source/TimecodeSynchronizerEditor/Private/UI/TimecodeSynchronizerEditorLevelToolbar.cpp b/Engine/Plugins/Media/TimecodeSynchronizer/Source/TimecodeSynchronizerEditor/Private/UI/TimecodeSynchronizerEditorLevelToolbar.cpp
--- a/Engine/Plugins/Media/TimecodeSynchronizer/Source/TimecodeSynchronizerEditor/Private/UI/TimecodeSynchronizerEditorLevelToolbar.cpp
+++ b/Engine/Plugins/Media/TimecodeSynchronizer/Source/TimecodeSynchronizerEditor/Private/UI/TimecodeSynchronizerEditorLevelToolbar.cpp
@@ -70,12 +70,37 @@ void FTimecodeSynchronizerEditorLevelToolbar::FillToolbar(FToolBarBuilder& Toolb
 {
 	ToolbarBuilder.BeginSection("TimecodeSynchronizer");
 	{
+		auto TooltipLambda = [this]()
+		{
+			if (!CurrentTimecodeSynchronizer.IsValid())
+			{
+				return LOCTEXT("EmptyTimecodeSynchronizer_ToolTip", "Select a Timecode Synchronizer to edit it.");
+			}
+			return FText::Format(LOCTEXT("TimecodeSynchronizer_ToolTip", "Edit '{0}'")
+				, FText::FromName(CurrentTimecodeSynchronizer->GetFName()));
+		};
+
+		// Add a button to edit the current timecode synchronizer
+		ToolbarBuilder.AddToolBarButton(
+			FUIAction(
+				FExecuteAction::CreateRaw(this, &FTimecodeSynchronizerEditorLevelToolbar::OpenCurrentTimecodeSynchronizer),
+				FCanExecuteAction::CreateLambda([this] { return CurrentTimecodeSynchronizer.IsValid(); }),
+				FIsActionChecked::CreateLambda([this] { return CurrentTimecodeSynchronizer.IsValid(); })
+			),
+			NAME_None,
+			LOCTEXT("TimecodeSynch_Label", "Timecode Synchronizer"),
+			MakeAttributeLambda(TooltipLambda),
+			FSlateIcon(FTimecodeSynchronizerEditorStyle::GetStyleSetName(), TEXT("Console"))
+		);
+
+		// Add a simple drop-down menu (no label, no icon for the drop-down button itself) that list the timecode synchronizer available
 		ToolbarBuilder.AddComboButton(
 			FUIAction(),
 			FOnGetContent::CreateRaw(this, &FTimecodeSynchronizerEditorLevelToolbar::GenerateMenuContent),
-			LOCTEXT("TimecodeSynch_Label", "Timecode Synchronizer"),
-			LOCTEXT("TimecodeSynch_ToolTip", "List of imecode Synchronizer available to the user for editing or creation."),
-			FSlateIcon(FTimecodeSynchronizerEditorStyle::GetStyleSetName(), TEXT("Console"))
+			FText::GetEmpty(),
+			LOCTEXT("TimecodeSynchButton_ToolTip", "List of Timecode Synchronizer available to the user for editing or creation."),
+			FSlateIcon(),
+			true
 		);
 	}
 	ToolbarBuilder.EndSection();
@@ -86,7 +111,7 @@ TSharedRef<SWidget> FTimecodeSynchronizerEditorLevelToolbar::GenerateMenuContent
 	const bool bShouldCloseWindowAfterMenuSelection = true;
 	FMenuBuilder MenuBuilder(bShouldCloseWindowAfterMenuSelection, nullptr);
 
-	MenuBuilder.BeginSection("TimecodeSynchronizer", LOCTEXT("TimecodeSynchronizerSection", "Timecode Synchronizer"));
+	MenuBuilder.BeginSection("TimecodeSynchronizer", LOCTEXT("NewTimecodeSynchronizerSection", "New"));
 	{
 		MenuBuilder.AddMenuEntry(
 			LOCTEXT("CreateMenuLabel", "New Empty Timecode Synchronizer"),
@@ -96,20 +121,18 @@ TSharedRef<SWidget> FTimecodeSynchronizerEditorLevelToolbar::GenerateMenuContent
 				FExecuteAction::CreateRaw(this, &FTimecodeSynchronizerEditorLevelToolbar::CreateNewTimecodeSynchronizer)
 			)
 		);
+	}
+	MenuBuilder.EndSection();
 
-		MenuBuilder.AddMenuSeparator();
-
+	MenuBuilder.BeginSection("TimecodeSynchronizer", LOCTEXT("TimecodeSynchronizerSection", "Timecode Synchronizer"));
+	{
 		const bool bIsTimecodeSynchronizerValid = CurrentTimecodeSynchronizer.IsValid();
 
 		MenuBuilder.AddSubMenu(
-			bIsTimecodeSynchronizerValid ? FText::Format(LOCTEXT("EditMenuLabel", "Open '{0}'"), FText::FromName(CurrentTimecodeSynchronizer->GetFName())) : LOCTEXT("SelectMenuLabel", "Select Profile"),
+			bIsTimecodeSynchronizerValid ? FText::FromName(CurrentTimecodeSynchronizer->GetFName()) : LOCTEXT("SelectMenuLabel", "Select a Timecode Synchronizer"),
 			LOCTEXT("SelectMenuTooltip", "Select the current timecode synchronizer for this editor."),
 			FNewMenuDelegate::CreateRaw(this, &FTimecodeSynchronizerEditorLevelToolbar::AddObjectSubMenu),
-			FUIAction(
-				FExecuteAction::CreateRaw(this, &FTimecodeSynchronizerEditorLevelToolbar::OpenCurrentTimecodeSynchronizer),
-				FCanExecuteAction(),
-				FIsActionChecked::CreateLambda([bIsTimecodeSynchronizerValid] { return bIsTimecodeSynchronizerValid; })
-			),
+			FUIAction(),
 			NAME_None,
 			EUserInterfaceActionType::RadioButton
 		);
