Set the property of MediaProfile to protected.

diff --git a/Engine/Plugins/Media/MediaFrameworkUtilities/Source/MediaFrameworkUtilities/Public/Profile/MediaProfile.h b/Engine/Plugins/Media/MediaFrameworkUtilities/Source/MediaFrameworkUtilities/Public/Profile/MediaProfile.h
--- a/Engine/Plugins/Media/MediaFrameworkUtilities/Source/MediaFrameworkUtilities/Public/Profile/MediaProfile.h
+++ b/Engine/Plugins/Media/MediaFrameworkUtilities/Source/MediaFrameworkUtilities/Public/Profile/MediaProfile.h
@@ -21,7 +21,7 @@ class MEDIAFRAMEWORKUTILITIES_API UMediaProfile : public UObject
 {
 	GENERATED_BODY()
 
-private:
+protected:
 
 	/** Media sources. */
 	UPROPERTY(EditAnywhere, Instanced, Category="Inputs", EditFixedSize, meta=(EditFixedOrder))
@@ -102,13 +102,13 @@ public:
 	 * Apply the media profile.
 	 * Will change the engine's timecode provider & custom time step and redirect the media profile source/output proxy for the correct media source/output.
 	 */
-	void Apply();
+	virtual void Apply();
 
 	/**
 	 * Reset the media profile.
 	 * Will reset the engine's timecode provider & custom time step and redirect the media profile source/output proxy for no media source/output.
 	 */
-	void Reset();
+	virtual void Reset();
 
 
 	/**
