- Potential crash in livelink deprecated path
- Fix missing unregister settings in livelink module

diff --git a/Engine/Plugins/Animation/LiveLink/Source/LiveLink/Private/LiveLinkClient.cpp b/Engine/Plugins/Animation/LiveLink/Source/LiveLink/Private/LiveLinkClient.cpp
--- a/Engine/Plugins/Animation/LiveLink/Source/LiveLink/Private/LiveLinkClient.cpp
+++ b/Engine/Plugins/Animation/LiveLink/Source/LiveLink/Private/LiveLinkClient.cpp
@@ -1163,8 +1163,11 @@ void FLiveLinkClient_Base_DEPRECATED::PushSubjectData(FGuid InSourceGuid, FName
 	NewData.WorldTime = InFrameData.WorldTime;
 	NewData.Transforms = InFrameData.Transforms;
 
+	//Always match FrameData property count to StaticData property count.
+	//If StaticData has more properties than current FrameData, set non existent properties to Infinity
+	//If StaticData has less properties than current FrameData, only use a subset of the incoming properties
 	int32 MaxNumberOfProperties = FMath::Min(NumberOfPropertyNames, InFrameData.CurveElements.Num());
-	NewData.PropertyValues.SetNumZeroed(MaxNumberOfProperties);
+	NewData.PropertyValues.SetNumZeroed(NumberOfPropertyNames);
 	for (int32 i = 0; i < MaxNumberOfProperties; ++i)
 	{
 		NewData.PropertyValues[i] = InFrameData.CurveElements[i].CurveValue;
diff --git a/Engine/Plugins/Animation/LiveLink/Source/LiveLinkEditor/Private/LiveLinkEditorModule.cpp b/Engine/Plugins/Animation/LiveLink/Source/LiveLinkEditor/Private/LiveLinkEditorModule.cpp
--- a/Engine/Plugins/Animation/LiveLink/Source/LiveLinkEditor/Private/LiveLinkEditorModule.cpp
+++ b/Engine/Plugins/Animation/LiveLink/Source/LiveLinkEditor/Private/LiveLinkEditorModule.cpp
@@ -215,7 +215,7 @@ private:
 		ISettingsModule* SettingsModule = FModuleManager::GetModulePtr<ISettingsModule>("Settings");
 		if (SettingsModule != nullptr)
 		{
-			SettingsModule->UnregisterSettings("Project", "Media", "MediaProfile");
+			SettingsModule->UnregisterSettings("Project", "Plugins", "LiveLink");
 		}
 	}
 
diff --git a/Engine/Plugins/Media/MediaFrameworkUtilities/Source/MediaFrameworkUtilities/Private/MediaBundleTimeSynchronizationSource.cpp b/Engine/Plugins/Media/MediaFrameworkUtilities/Source/MediaFrameworkUtilities/Private/MediaBundleTimeSynchronizationSource.cpp
--- a/Engine/Plugins/Media/MediaFrameworkUtilities/Source/MediaFrameworkUtilities/Private/MediaBundleTimeSynchronizationSource.cpp
+++ b/Engine/Plugins/Media/MediaFrameworkUtilities/Source/MediaFrameworkUtilities/Private/MediaBundleTimeSynchronizationSource.cpp
@@ -47,10 +47,10 @@ void UMediaBundleTimeSynchronizationSource::PostEditChangeProperty(FPropertyChan
 		if (bUseForSynchronization && MediaBundle && MediaBundle->GetMediaSource())
 		{
 			UTimeSynchronizableMediaSource* SynchronizableMediaSource = TimeSynchronizationSource::GetTimeSynchronizableMediaSource(MediaBundle);
-			if (SynchronizableMediaSource == nullptr || SynchronizableMediaSource->bUseTimeSynchronization)
+			if (SynchronizableMediaSource == nullptr || !SynchronizableMediaSource->bUseTimeSynchronization)
 			{
 				// Warn the user that the MediaSource he just added isn't set for using time synchronization
-				UE_LOG(LogMediaFrameworkUtilities, Warning, TEXT("MediaSource %s doesn't have the option to use time synchronization enabled."), *MediaBundle->MediaSource->GetName());
+				UE_LOG(LogMediaFrameworkUtilities, Warning, TEXT("MediaSource '%s' doesn't have the option to use time synchronization enabled."), *MediaBundle->MediaSource->GetName());
 			}
 		}
 	}
@@ -176,7 +176,7 @@ bool UMediaBundleTimeSynchronizationSource::Open(const FTimeSynchronizationOpenD
 			UTimeSynchronizableMediaSource* SynchronizableMediaSource = TimeSynchronizationSource::GetTimeSynchronizableMediaSource(MediaBundle);
 			if (bUseForSynchronization && (SynchronizableMediaSource == nullptr || !SynchronizableMediaSource->bUseTimeSynchronization))
 			{
-				UE_LOG(LogMediaFrameworkUtilities, Error, TEXT("MediaBundle %s doesn't support timecode synchronization"), *MediaBundle->GetName());
+				UE_LOG(LogMediaFrameworkUtilities, Error, TEXT("MediaBundle '%s' doesn't support timecode synchronization"), *MediaBundle->GetName());
 			}
 			else
 			{
