- Fixed MediaProfile typo

diff --git a/Engine/Plugins/Media/MediaFrameworkUtilities/Source/MediaFrameworkUtilitiesEditor/Private/Profile/MediaProfileSettingsCustomization.cpp b/Engine/Plugins/Media/MediaFrameworkUtilities/Source/MediaFrameworkUtilitiesEditor/Private/Profile/MediaProfileSettingsCustomization.cpp
--- a/Engine/Plugins/Media/MediaFrameworkUtilities/Source/MediaFrameworkUtilitiesEditor/Private/Profile/MediaProfileSettingsCustomization.cpp
+++ b/Engine/Plugins/Media/MediaFrameworkUtilities/Source/MediaFrameworkUtilitiesEditor/Private/Profile/MediaProfileSettingsCustomization.cpp
@@ -245,7 +245,7 @@ void FMediaProfileSettingsCustomization::Configure(const FMediaProfileSettingsCu
 	{
 		TArray<UPackage*> PackagesToSave;
 
-		FScopedSlowTask Progress(100.0f, LOCTEXT("StartWork", "Configure Medie Profile"));
+		FScopedSlowTask Progress(100.0f, LOCTEXT("StartWork", "Configure Media Profile"));
 		Progress.MakeDialog(true);
 
 		// Create the sources proxies at the location
