Garbage collect the media proxies from the media setting outside of the CDO.

diff --git a/Engine/Plugins/Media/MediaFrameworkUtilities/Source/MediaFrameworkUtilities/Private/Profile/MediaProfile.cpp b/Engine/Plugins/Media/MediaFrameworkUtilities/Source/MediaFrameworkUtilities/Private/Profile/MediaProfile.cpp
--- a/Engine/Plugins/Media/MediaFrameworkUtilities/Source/MediaFrameworkUtilities/Private/Profile/MediaProfile.cpp
+++ b/Engine/Plugins/Media/MediaFrameworkUtilities/Source/MediaFrameworkUtilities/Private/Profile/MediaProfile.cpp
@@ -12,7 +12,7 @@
 #include "MediaAssets/ProxyMediaSource.h"
 #include "MediaOutput.h"
 #include "MediaSource.h"
-#include "Profile/MediaProfileSettings.h"
+#include "Profile/IMediaProfileManager.h"
 
 
 UMediaSource* UMediaProfile::GetMediaSource(int32 Index) const
@@ -75,7 +75,7 @@ void UMediaProfile::Apply()
 	FixNumSourcesAndOutputs();
 
 	{
-		TArray<UProxyMediaSource*> SourceProxies = GetDefault<UMediaProfileSettings>()->GetAllMediaSourceProxy();
+		TArray<UProxyMediaSource*> SourceProxies = IMediaProfileManager::Get().GetAllMediaSourceProxy();
 		check(SourceProxies.Num() == MediaSources.Num());
 		for (int32 Index = 0; Index < MediaSources.Num(); ++Index)
 		{
@@ -88,7 +88,7 @@ void UMediaProfile::Apply()
 	}
 
 	{
-		TArray<UProxyMediaOutput*> OutputProxies = GetDefault<UMediaProfileSettings>()->GetAllMediaOutputProxy();
+		TArray<UProxyMediaOutput*> OutputProxies = IMediaProfileManager::Get().GetAllMediaOutputProxy();
 		check(OutputProxies.Num() == MediaOutputs.Num());
 		for (int32 Index = 0; Index < MediaOutputs.Num(); ++Index)
 		{
@@ -144,7 +144,7 @@ void UMediaProfile::Reset()
 
 	{
 		// Reset the source proxies
-		TArray<UProxyMediaSource*> SourceProxies = GetDefault<UMediaProfileSettings>()->GetAllMediaSourceProxy();
+		TArray<UProxyMediaSource*> SourceProxies = IMediaProfileManager::Get().GetAllMediaSourceProxy();
 		for (UProxyMediaSource* Proxy : SourceProxies)
 		{
 			if (Proxy)
@@ -156,7 +156,7 @@ void UMediaProfile::Reset()
 
 	{
 		// Reset the output proxies
-		TArray<UProxyMediaOutput*> OutputProxies = GetDefault<UMediaProfileSettings>()->GetAllMediaOutputProxy();
+		TArray<UProxyMediaOutput*> OutputProxies = IMediaProfileManager::Get().GetAllMediaOutputProxy();
 		for (UProxyMediaOutput* Proxy : OutputProxies)
 		{
 			if (Proxy)
@@ -193,14 +193,14 @@ void UMediaProfile::Reset()
 
 void UMediaProfile::FixNumSourcesAndOutputs()
 {
-	const int32 NumSourceProxies = GetDefault<UMediaProfileSettings>()->GetAllMediaSourceProxy().Num();
+	const int32 NumSourceProxies = IMediaProfileManager::Get().GetAllMediaSourceProxy().Num();
 	if (MediaSources.Num() != NumSourceProxies)
 	{
 		MediaSources.SetNumZeroed(NumSourceProxies);
 		Modify();
 	}
 
-	const int32 NumOutputProxies = GetDefault<UMediaProfileSettings>()->GetAllMediaOutputProxy().Num();
+	const int32 NumOutputProxies = IMediaProfileManager::Get().GetAllMediaOutputProxy().Num();
 	if (MediaOutputs.Num() != NumOutputProxies)
 	{
 		Modify();
diff --git a/Engine/Plugins/Media/MediaFrameworkUtilities/Source/MediaFrameworkUtilities/Private/Profile/MediaProfileBlueprintLibrary.cpp b/Engine/Plugins/Media/MediaFrameworkUtilities/Source/MediaFrameworkUtilities/Private/Profile/MediaProfileBlueprintLibrary.cpp
--- a/Engine/Plugins/Media/MediaFrameworkUtilities/Source/MediaFrameworkUtilities/Private/Profile/MediaProfileBlueprintLibrary.cpp
+++ b/Engine/Plugins/Media/MediaFrameworkUtilities/Source/MediaFrameworkUtilities/Private/Profile/MediaProfileBlueprintLibrary.cpp
@@ -4,7 +4,6 @@
 
 #include "Features/IModularFeatures.h"
 #include "Profile/IMediaProfileManager.h"
-#include "Profile/MediaProfileSettings.h"
 
 
 UMediaProfile* UMediaProfileBlueprintLibrary::GetMediaProfile()
@@ -21,11 +20,11 @@ void UMediaProfileBlueprintLibrary::SetMediaProfile(UMediaProfile* MediaProfile)
 
 TArray<UProxyMediaSource*> UMediaProfileBlueprintLibrary::GetAllMediaSourceProxy()
 {
-	return GetDefault<UMediaProfileSettings>()->GetAllMediaSourceProxy();
+	return IMediaProfileManager::Get().GetAllMediaSourceProxy();
 }
 
 
 TArray<UProxyMediaOutput*> UMediaProfileBlueprintLibrary::GetAllMediaOutputProxy()
 {
-	return GetDefault<UMediaProfileSettings>()->GetAllMediaOutputProxy();
+	return IMediaProfileManager::Get().GetAllMediaOutputProxy();
 }
diff --git a/Engine/Plugins/Media/MediaFrameworkUtilities/Source/MediaFrameworkUtilities/Private/Profile/MediaProfileManager.cpp b/Engine/Plugins/Media/MediaFrameworkUtilities/Source/MediaFrameworkUtilities/Private/Profile/MediaProfileManager.cpp
--- a/Engine/Plugins/Media/MediaFrameworkUtilities/Source/MediaFrameworkUtilities/Private/Profile/MediaProfileManager.cpp
+++ b/Engine/Plugins/Media/MediaFrameworkUtilities/Source/MediaFrameworkUtilities/Private/Profile/MediaProfileManager.cpp
@@ -19,12 +19,50 @@ IMediaProfileManager& IMediaProfileManager::Get()
 FMediaProfileManager::FMediaProfileManager()
 	: CurrentMediaProfile(nullptr)
 {
+	if (UObjectInitialized())
+	{
+		MediaSourceProxies = GetDefault<UMediaProfileSettings>()->LoadMediaSourceProxies();
+		MediaOutputProxies = GetDefault<UMediaProfileSettings>()->LoadMediaOutputProxies();
+#if WITH_EDITOR
+		GetMutableDefault<UMediaProfileSettings>()->OnMediaProxiesChanged.AddRaw(this, &FMediaProfileManager::OnMediaProxiesChanged);
+#endif
+	}
+}
+
+
+FMediaProfileManager::~FMediaProfileManager()
+{
+#if WITH_EDITOR
+	if (UObjectInitialized())
+	{
+		GetMutableDefault<UMediaProfileSettings>()->OnMediaProxiesChanged.RemoveAll(this);
+	}
+#endif
+}
+
+
+void FMediaProfileManager::AddReferencedObjects(FReferenceCollector& Collector)
+{
+	Collector.AddReferencedObject(CurrentMediaProfile);
+	Collector.AddReferencedObjects(MediaSourceProxies);
+	Collector.AddReferencedObjects(MediaOutputProxies);
 }
 
 
 UMediaProfile* FMediaProfileManager::GetCurrentMediaProfile() const
 {
-	return CurrentMediaProfile.Get();
+	return CurrentMediaProfile;
+}
+
+TArray<UProxyMediaSource*> FMediaProfileManager::GetAllMediaSourceProxy() const
+{
+	return MediaSourceProxies;
+}
+
+
+TArray<UProxyMediaOutput*> FMediaProfileManager::GetAllMediaOutputProxy() const
+{
+	return MediaOutputProxies;
 }
 
 
@@ -32,7 +70,7 @@ void FMediaProfileManager::SetCurrentMediaProfile(UMediaProfile* InMediaProfile)
 {
 	bool bRemoveDelegate = false;
 	bool bAddDelegate = false;
-	UMediaProfile* Previous = CurrentMediaProfile.Get();
+	UMediaProfile* Previous = CurrentMediaProfile;
 	if (InMediaProfile != Previous)
 	{
 		if (Previous)
@@ -47,23 +85,9 @@ void FMediaProfileManager::SetCurrentMediaProfile(UMediaProfile* InMediaProfile)
 			bAddDelegate = true;
 		}
 
-		CurrentMediaProfile.Reset(InMediaProfile);
+		CurrentMediaProfile = InMediaProfile;
 		MediaProfileChangedDelegate.Broadcast(Previous, InMediaProfile);
 	}
-
-#if WITH_EDITOR
-	if (UObjectInitialized())
-	{
-		if (bAddDelegate && !bRemoveDelegate)
-		{
-			GetMutableDefault<UMediaProfileSettings>()->OnMediaProxiesChanged.AddRaw(this, &FMediaProfileManager::OnMediaProxiesChanged);
-		}
-		else if (bRemoveDelegate && !bAddDelegate)
-		{
-			GetMutableDefault<UMediaProfileSettings>()->OnMediaProxiesChanged.RemoveAll(this);
-		}
-	}
-#endif
 }
 
 
@@ -75,11 +99,13 @@ FMediaProfileManager::FOnMediaProfileChanged& FMediaProfileManager::OnMediaProfi
 #if WITH_EDITOR
 void FMediaProfileManager::OnMediaProxiesChanged()
 {
-	UMediaProfile* CurrentMediaProfilePtr = CurrentMediaProfile.Get();
-	if (CurrentMediaProfilePtr)
+	MediaSourceProxies = GetDefault<UMediaProfileSettings>()->LoadMediaSourceProxies();
+	MediaOutputProxies = GetDefault<UMediaProfileSettings>()->LoadMediaOutputProxies();
+
+	if (CurrentMediaProfile)
 	{
-		CurrentMediaProfilePtr->Reset();
-		CurrentMediaProfilePtr->Apply();
+		CurrentMediaProfile->Reset();
+		CurrentMediaProfile->Apply();
 	}
 }
 #endif
diff --git a/Engine/Plugins/Media/MediaFrameworkUtilities/Source/MediaFrameworkUtilities/Private/Profile/MediaProfileManager.h b/Engine/Plugins/Media/MediaFrameworkUtilities/Source/MediaFrameworkUtilities/Private/Profile/MediaProfileManager.h
--- a/Engine/Plugins/Media/MediaFrameworkUtilities/Source/MediaFrameworkUtilities/Private/Profile/MediaProfileManager.h
+++ b/Engine/Plugins/Media/MediaFrameworkUtilities/Source/MediaFrameworkUtilities/Private/Profile/MediaProfileManager.h
@@ -3,26 +3,37 @@
 #pragma once
 
 #include "Profile/IMediaProfileManager.h"
-#include "UObject/StrongObjectPtr.h"
+#include "UObject/GCObject.h"
 
 class UMediaProfile;
 class UProxyMediaSource;
 class UProxyMediaOutput;
 
-class FMediaProfileManager : public IMediaProfileManager
+class FMediaProfileManager : public IMediaProfileManager, public FGCObject
 {
 public:
 	FMediaProfileManager();
+	virtual ~FMediaProfileManager();
 
+	//~ Begin IMediaProfileManager Interface
 	virtual UMediaProfile* GetCurrentMediaProfile() const override;
 	virtual void SetCurrentMediaProfile(UMediaProfile* InMediaProfile) override;
+	TArray<UProxyMediaSource*> GetAllMediaSourceProxy() const override;
+	TArray<UProxyMediaOutput*> GetAllMediaOutputProxy() const override;
 	virtual FOnMediaProfileChanged& OnMediaProfileChanged() override;
+	//~ End IMediaProfileManager Interface
+
+	//~ Begin FGCObject Interface
+	virtual void AddReferencedObjects(FReferenceCollector& Collector) override;
+	//~End FGCObject Interface
 
 private:
 #if WITH_EDITOR
 	void OnMediaProxiesChanged();
 #endif
 
-	TStrongObjectPtr<UMediaProfile> CurrentMediaProfile;
+	TArray<UProxyMediaSource*> MediaSourceProxies;
+	TArray<UProxyMediaOutput*> MediaOutputProxies;
+	UMediaProfile* CurrentMediaProfile;
 	FOnMediaProfileChanged MediaProfileChangedDelegate;
 };
diff --git a/Engine/Plugins/Media/MediaFrameworkUtilities/Source/MediaFrameworkUtilities/Private/Profile/MediaProfileSettings.cpp b/Engine/Plugins/Media/MediaFrameworkUtilities/Source/MediaFrameworkUtilities/Private/Profile/MediaProfileSettings.cpp
--- a/Engine/Plugins/Media/MediaFrameworkUtilities/Source/MediaFrameworkUtilities/Private/Profile/MediaProfileSettings.cpp
+++ b/Engine/Plugins/Media/MediaFrameworkUtilities/Source/MediaFrameworkUtilities/Private/Profile/MediaProfileSettings.cpp
@@ -8,30 +8,27 @@
 #include "Profile/MediaProfile.h"
 
 
-TArray<UProxyMediaSource*> UMediaProfileSettings::GetAllMediaSourceProxy() const
+TArray<UProxyMediaSource*> UMediaProfileSettings::LoadMediaSourceProxies() const
 {
-	if (MediaSourceProxyPtr.Num() != MediaSourceProxy.Num())
+	TArray<UProxyMediaSource*> MediaSourceProxyPtr;
+	MediaSourceProxyPtr.Reset(MediaSourceProxy.Num());
+
+	for (const TSoftObjectPtr<UProxyMediaSource>& Proxy : MediaSourceProxy)
 	{
-		MediaSourceProxyPtr.Reset(MediaSourceProxy.Num());
-		for (const TSoftObjectPtr<UProxyMediaSource>& Proxy : MediaSourceProxy)
-		{
-			MediaSourceProxyPtr.Add(Proxy.LoadSynchronous());
-		}
+		MediaSourceProxyPtr.Add(Proxy.LoadSynchronous());
 	}
 	return MediaSourceProxyPtr;
 }
 
 
-TArray<UProxyMediaOutput*> UMediaProfileSettings::GetAllMediaOutputProxy() const
+TArray<UProxyMediaOutput*> UMediaProfileSettings::LoadMediaOutputProxies() const
 {
-	if (MediaOutputProxyPtr.Num() != MediaOutputProxy.Num())
-	{
-		MediaOutputProxyPtr.Reset(MediaOutputProxy.Num());
+	TArray<UProxyMediaOutput*> MediaOutputProxyPtr;
+	MediaOutputProxyPtr.Reset(MediaOutputProxy.Num());
 
-		for (const TSoftObjectPtr<UProxyMediaOutput>& Proxy : MediaOutputProxy)
-		{
-			MediaOutputProxyPtr.Add(Proxy.LoadSynchronous());
-		}
+	for (const TSoftObjectPtr<UProxyMediaOutput>& Proxy : MediaOutputProxy)
+	{
+		MediaOutputProxyPtr.Add(Proxy.LoadSynchronous());
 	}
 	return MediaOutputProxyPtr;
 }
@@ -47,7 +44,6 @@ UMediaProfile* UMediaProfileSettings::GetStartupMediaProfile() const
 void UMediaProfileSettings::SetMediaSourceProxy(const TArray<UProxyMediaSource*>& InProxies)
 {
 	MediaSourceProxy.Reset();
-	MediaSourceProxyPtr = InProxies;
 	for (UProxyMediaSource* Proxy : InProxies)
 	{
 		MediaSourceProxy.Add(Proxy);
@@ -59,7 +55,6 @@ void UMediaProfileSettings::SetMediaSourceProxy(const TArray<UProxyMediaSource*>
 void UMediaProfileSettings::SetMediaOutputProxy(const TArray<UProxyMediaOutput*>& InProxies)
 {
 	MediaOutputProxy.Reset();
-	MediaOutputProxyPtr = InProxies;
 	for (UProxyMediaOutput* Proxy : InProxies)
 	{
 		MediaOutputProxy.Add(Proxy);
@@ -74,11 +69,6 @@ void UMediaProfileSettings::PostEditChangeProperty(struct FPropertyChangedEvent&
 		|| GET_MEMBER_NAME_CHECKED(ThisClass, MediaSourceProxy) == InPropertyChangedEvent.GetPropertyName())
 	{
 		// Recache the proxies ptr
-		MediaSourceProxyPtr.Reset(MediaSourceProxy.Num());
-		MediaOutputProxyPtr.Reset(MediaOutputProxy.Num());
-		GetAllMediaSourceProxy();
-		GetAllMediaOutputProxy();
-
 		OnMediaProxiesChanged.Broadcast();
 	}
 
diff --git a/Engine/Plugins/Media/MediaFrameworkUtilities/Source/MediaFrameworkUtilities/Public/Profile/IMediaProfileManager.h b/Engine/Plugins/Media/MediaFrameworkUtilities/Source/MediaFrameworkUtilities/Public/Profile/IMediaProfileManager.h
--- a/Engine/Plugins/Media/MediaFrameworkUtilities/Source/MediaFrameworkUtilities/Public/Profile/IMediaProfileManager.h
+++ b/Engine/Plugins/Media/MediaFrameworkUtilities/Source/MediaFrameworkUtilities/Public/Profile/IMediaProfileManager.h
@@ -5,6 +5,8 @@
 #include "CoreMinimal.h"
 
 class UMediaProfile;
+class UProxyMediaSource;
+class UProxyMediaOutput;
 
 class MEDIAFRAMEWORKUTILITIES_API IMediaProfileManager
 {
@@ -20,6 +22,12 @@ public:
 	/** Set the current profile used by the manager. */
 	virtual void SetCurrentMediaProfile(UMediaProfile* InMediaProfile) = 0;
 
+	/** Get all the media source proxy. */
+	virtual TArray<UProxyMediaSource*> GetAllMediaSourceProxy() const = 0;
+
+	/** Get all the media output proxy. */
+	virtual TArray<UProxyMediaOutput*> GetAllMediaOutputProxy() const = 0;
+
 	/** Delegate type for media profile changed event */
 	DECLARE_MULTICAST_DELEGATE_TwoParams(FOnMediaProfileChanged, UMediaProfile* /*Preivous*/, UMediaProfile* /*New*/);
 
diff --git a/Engine/Plugins/Media/MediaFrameworkUtilities/Source/MediaFrameworkUtilities/Public/Profile/MediaProfileSettings.h b/Engine/Plugins/Media/MediaFrameworkUtilities/Source/MediaFrameworkUtilities/Public/Profile/MediaProfileSettings.h
--- a/Engine/Plugins/Media/MediaFrameworkUtilities/Source/MediaFrameworkUtilities/Public/Profile/MediaProfileSettings.h
+++ b/Engine/Plugins/Media/MediaFrameworkUtilities/Source/MediaFrameworkUtilities/Public/Profile/MediaProfileSettings.h
@@ -42,12 +42,6 @@ private:
 	UPROPERTY(AdvancedDisplay, Config, EditAnywhere, Category="MediaProfile")
 	TArray<TSoftObjectPtr<UProxyMediaOutput>> MediaOutputProxy;
 
-	UPROPERTY(Transient)
-	mutable TArray<UProxyMediaSource*> MediaSourceProxyPtr;
-
-	UPROPERTY(Transient)
-	mutable TArray<UProxyMediaOutput*> MediaOutputProxyPtr;
-
 	/**
 	 * The media profile to use at startup.
 	 * @note The media profile can be overridden in the editor by user.
@@ -62,14 +56,14 @@ public:
 	 *
 	 * @return The an array of the media source proxy.
 	 */
-	TArray<UProxyMediaSource*> GetAllMediaSourceProxy() const;
+	TArray<UProxyMediaSource*> LoadMediaSourceProxies() const;
 
 	/**
 	 * Get all the media output proxy.
 	 *
 	 * @return The an array of the media output proxy.
 	 */
-	TArray<UProxyMediaOutput*> GetAllMediaOutputProxy() const;
+	TArray<UProxyMediaOutput*> LoadMediaOutputProxies() const;
 
 	/**
 	 * Get the media profile used by the engine.
diff --git a/Engine/Plugins/Media/MediaFrameworkUtilities/Source/MediaFrameworkUtilitiesEditor/Private/Profile/MediaProfileCustomization.cpp b/Engine/Plugins/Media/MediaFrameworkUtilities/Source/MediaFrameworkUtilitiesEditor/Private/Profile/MediaProfileCustomization.cpp
--- a/Engine/Plugins/Media/MediaFrameworkUtilities/Source/MediaFrameworkUtilitiesEditor/Private/Profile/MediaProfileCustomization.cpp
+++ b/Engine/Plugins/Media/MediaFrameworkUtilities/Source/MediaFrameworkUtilitiesEditor/Private/Profile/MediaProfileCustomization.cpp
@@ -8,7 +8,6 @@
 #include "MediaAssets/ProxyMediaSource.h"
 #include "Profile/IMediaProfileManager.h"
 #include "Profile/MediaProfile.h"
-#include "Profile/MediaProfileSettings.h"
 
 
 #define LOCTEXT_NAMESPACE "MediaProfileCustomization"
@@ -26,8 +25,8 @@ void FMediaProfileCustomization::CustomizeDetails(IDetailLayoutBuilder& DetailLa
 {
 	Super::CustomizeDetails(DetailLayout);
 
-	const TArray<UProxyMediaSource*> SourceProxies = GetDefault<UMediaProfileSettings>()->GetAllMediaSourceProxy();
-	const TArray<UProxyMediaOutput*> OutputProxies = GetDefault<UMediaProfileSettings>()->GetAllMediaOutputProxy();
+	const TArray<UProxyMediaSource*> SourceProxies = IMediaProfileManager::Get().GetAllMediaSourceProxy();
+	const TArray<UProxyMediaOutput*> OutputProxies = IMediaProfileManager::Get().GetAllMediaOutputProxy();
 
 	const TArray<TWeakObjectPtr<UObject>>& SelectedObjects = DetailLayout.GetSelectedObjects();
 	for (const TWeakObjectPtr<UObject>& CurrentWeakObject : SelectedObjects)
diff --git a/Engine/Plugins/Media/MediaFrameworkUtilities/Source/MediaFrameworkUtilitiesEditor/Private/Profile/MediaProfileSettingsCustomization.cpp b/Engine/Plugins/Media/MediaFrameworkUtilities/Source/MediaFrameworkUtilitiesEditor/Private/Profile/MediaProfileSettingsCustomization.cpp
--- a/Engine/Plugins/Media/MediaFrameworkUtilities/Source/MediaFrameworkUtilitiesEditor/Private/Profile/MediaProfileSettingsCustomization.cpp
+++ b/Engine/Plugins/Media/MediaFrameworkUtilities/Source/MediaFrameworkUtilitiesEditor/Private/Profile/MediaProfileSettingsCustomization.cpp
@@ -45,8 +45,8 @@ void FMediaProfileSettingsCustomization::CustomizeDetails(IDetailLayoutBuilder&
 	const UMediaProfileSettings* MediaProfileSettings = GetDefault<UMediaProfileSettings>();
 	check(MediaProfileSettings);
 
-	const TArray<UProxyMediaOutput*> OutputProxies = MediaProfileSettings->GetAllMediaOutputProxy();
-	const TArray<UProxyMediaSource*> SourceProxies = MediaProfileSettings->GetAllMediaSourceProxy();
+	const TArray<UProxyMediaOutput*> OutputProxies = MediaProfileSettings->LoadMediaOutputProxies();
+	const TArray<UProxyMediaSource*> SourceProxies = MediaProfileSettings->LoadMediaSourceProxies();
 
 	bNotConfigured = OutputProxies.Num() == 0 && SourceProxies.Num() == 0;
 	bNeedFixup = false;
