Delay the load of media proxies used by the media profile to prevent them from beeing tag as disregard by the GC.

diff --git a/Engine/Plugins/Media/MediaFrameworkUtilities/Source/MediaFrameworkUtilities/Private/Profile/MediaProfileManager.cpp b/Engine/Plugins/Media/MediaFrameworkUtilities/Source/MediaFrameworkUtilities/Private/Profile/MediaProfileManager.cpp
--- a/Engine/Plugins/Media/MediaFrameworkUtilities/Source/MediaFrameworkUtilities/Private/Profile/MediaProfileManager.cpp
+++ b/Engine/Plugins/Media/MediaFrameworkUtilities/Source/MediaFrameworkUtilities/Private/Profile/MediaProfileManager.cpp
@@ -2,9 +2,11 @@
 #include "Profile/MediaProfileManager.h"
 #include "MediaFrameworkUtilitiesModule.h"
 
+#include "Engine/Engine.h"
 #include "Modules/ModuleManager.h"
 #include "MediaAssets/ProxyMediaOutput.h"
 #include "MediaAssets/ProxyMediaSource.h"
+#include "Misc/App.h"
 #include "Profile/MediaProfile.h"
 #include "Profile/MediaProfileSettings.h"
 
@@ -21,8 +23,12 @@ FMediaProfileManager::FMediaProfileManager()
 {
 	if (UObjectInitialized())
 	{
-		MediaSourceProxies = GetDefault<UMediaProfileSettings>()->LoadMediaSourceProxies();
-		MediaOutputProxies = GetDefault<UMediaProfileSettings>()->LoadMediaOutputProxies();
+		if (FApp::CanEverRender() && GEngine && GEngine->IsInitialized())
+		{
+			MediaSourceProxies = GetDefault<UMediaProfileSettings>()->LoadMediaSourceProxies();
+			MediaOutputProxies = GetDefault<UMediaProfileSettings>()->LoadMediaOutputProxies();
+		}
+
 #if WITH_EDITOR
 		GetMutableDefault<UMediaProfileSettings>()->OnMediaProxiesChanged.AddRaw(this, &FMediaProfileManager::OnMediaProxiesChanged);
 #endif
@@ -54,6 +60,7 @@ UMediaProfile* FMediaProfileManager::GetCurrentMediaProfile() const
 	return CurrentMediaProfile;
 }
 
+
 TArray<UProxyMediaSource*> FMediaProfileManager::GetAllMediaSourceProxy() const
 {
 	return MediaSourceProxies;
@@ -68,6 +75,9 @@ TArray<UProxyMediaOutput*> FMediaProfileManager::GetAllMediaOutputProxy() const
 
 void FMediaProfileManager::SetCurrentMediaProfile(UMediaProfile* InMediaProfile)
 {
+	MediaSourceProxies = GetDefault<UMediaProfileSettings>()->LoadMediaSourceProxies();
+	MediaOutputProxies = GetDefault<UMediaProfileSettings>()->LoadMediaOutputProxies();
+
 	bool bRemoveDelegate = false;
 	bool bAddDelegate = false;
 	UMediaProfile* Previous = CurrentMediaProfile;
@@ -96,6 +106,7 @@ FMediaProfileManager::FOnMediaProfileChanged& FMediaProfileManager::OnMediaProfi
 	return MediaProfileChangedDelegate;
 }
 
+
 #if WITH_EDITOR
 void FMediaProfileManager::OnMediaProxiesChanged()
 {
