More CIS fix #rb none

diff --git a/Engine/Plugins/Editor/MeshEditor/Source/MeshEditor/MeshEditorMode.cpp b/Engine/Plugins/Editor/MeshEditor/Source/MeshEditor/MeshEditorMode.cpp
--- a/Engine/Plugins/Editor/MeshEditor/Source/MeshEditor/MeshEditorMode.cpp
+++ b/Engine/Plugins/Editor/MeshEditor/Source/MeshEditor/MeshEditorMode.cpp
@@ -960,13 +960,15 @@ void FMeshEditorMode::Enter()
 	{
 		const TSharedRef<ILevelEditor>& LevelEditor = LevelEditorModule.GetFirstLevelEditor().ToSharedRef();
 
+
 		// Do we have an active perspective viewport that is valid for VR?  If so, go ahead and use that.
 		TSharedPtr<FEditorViewportClient> ViewportClient;
 		{
 			TSharedPtr<IAssetViewport> ActiveLevelViewport = LevelEditor->GetActiveViewportInterface();
 			if (ActiveLevelViewport.IsValid())
 			{
-				ViewportClient = StaticCastSharedRef<SLevelViewport>(ActiveLevelViewport->AsWidget())->GetViewportClient();
+				const TSharedRef<SLevelViewport> LevelViewportWidget = StaticCastSharedRef<SLevelViewport>(ActiveLevelViewport->AsWidget());
+				ViewportClient = MakeShareable(&LevelViewportWidget->GetViewportClient());
 			}
 		}
 
diff --git a/Engine/Plugins/Media/MediaFrameworkUtilities/Source/MediaFrameworkUtilitiesEditor/Private/CaptureTab/SMediaFrameworkCaptureOutputWidget.cpp b/Engine/Plugins/Media/MediaFrameworkUtilities/Source/MediaFrameworkUtilitiesEditor/Private/CaptureTab/SMediaFrameworkCaptureOutputWidget.cpp
--- a/Engine/Plugins/Media/MediaFrameworkUtilities/Source/MediaFrameworkUtilitiesEditor/Private/CaptureTab/SMediaFrameworkCaptureOutputWidget.cpp
+++ b/Engine/Plugins/Media/MediaFrameworkUtilities/Source/MediaFrameworkUtilitiesEditor/Private/CaptureTab/SMediaFrameworkCaptureOutputWidget.cpp
@@ -607,7 +607,7 @@ void SMediaFrameworkCaptureCurrentViewportWidget::StartOutput()
 			{
 				// Find a editor viewport
 				FLevelEditorModule& LevelEditorModule = FModuleManager::Get().GetModuleChecked<FLevelEditorModule>(TEXT("LevelEditor"));
-				TSharedPtr<IAssetViewport> LevelViewportInterface = LevelEditorModule.GetFirstActiveViewport();
+				TSharedPtr<SLevelViewport> LevelViewportInterface = LevelEditorModule.GetFirstActiveLevelViewport();
 				if (LevelViewportInterface.IsValid())
 				{
 					LevelViewport = LevelViewportInterface;
@@ -709,7 +709,7 @@ void SMediaFrameworkCaptureCurrentViewportWidget::ShutdownViewport()
 	TSharedPtr<IAssetViewport> LevelViewportPin = LevelViewport.Pin();
 	if (LevelViewportPin.IsValid() && LevelViewportPin->GetSharedActiveViewport() == EditorSceneViewportPin)
 	{
-		FLevelEditorViewportClient& ViewportClient = LevelViewportPin->GetLevelViewportClient();
+		FEditorViewportClient& ViewportClient = LevelViewportPin->GetViewportClient();
 
 		// Reset settings
 		ViewportClient.SetRealtime(ViewportClientFlags.bRealTime);
