CIS fix #rb none

diff --git a/Engine/Plugins/Media/MediaFrameworkUtilities/Source/MediaFrameworkUtilitiesEditor/Private/CaptureTab/SMediaFrameworkCaptureOutputWidget.cpp b/Engine/Plugins/Media/MediaFrameworkUtilities/Source/MediaFrameworkUtilitiesEditor/Private/CaptureTab/SMediaFrameworkCaptureOutputWidget.cpp
--- a/Engine/Plugins/Media/MediaFrameworkUtilities/Source/MediaFrameworkUtilitiesEditor/Private/CaptureTab/SMediaFrameworkCaptureOutputWidget.cpp
+++ b/Engine/Plugins/Media/MediaFrameworkUtilities/Source/MediaFrameworkUtilitiesEditor/Private/CaptureTab/SMediaFrameworkCaptureOutputWidget.cpp
@@ -607,13 +607,13 @@ void SMediaFrameworkCaptureCurrentViewportWidget::StartOutput()
 			{
 				// Find a editor viewport
 				FLevelEditorModule& LevelEditorModule = FModuleManager::Get().GetModuleChecked<FLevelEditorModule>(TEXT("LevelEditor"));
-				TSharedPtr<SLevelViewport> LevelViewportInterface = LevelEditorModule.GetFirstActiveLevelViewport();
-				if (LevelViewportInterface.IsValid())
+				TSharedPtr<IAssetViewport> ViewportInterface = LevelEditorModule.GetFirstActiveViewport();
+				if (ViewportInterface.IsValid())
 				{
-					LevelViewport = LevelViewportInterface;
-					SceneViewport = LevelViewportInterface->GetSharedActiveViewport();
+					LevelViewport = ViewportInterface;
+					SceneViewport = ViewportInterface->GetSharedActiveViewport();
 
-					FLevelEditorViewportClient& ViewportClient = LevelViewportInterface->GetLevelViewportClient();
+					FLevelEditorViewportClient& ViewportClient = ViewportInterface->GetViewportClient();
 
 					// Save settings
 					ViewportClientFlags.bRealTime = ViewportClient.IsRealtime();
