Fix build
#rnx

diff --git a/Engine/Plugins/Media/MediaFrameworkUtilities/Source/MediaFrameworkUtilitiesEditor/Private/CaptureTab/SMediaFrameworkCaptureOutputWidget.cpp b/Engine/Plugins/Media/MediaFrameworkUtilities/Source/MediaFrameworkUtilitiesEditor/Private/CaptureTab/SMediaFrameworkCaptureOutputWidget.cpp
--- a/Engine/Plugins/Media/MediaFrameworkUtilities/Source/MediaFrameworkUtilitiesEditor/Private/CaptureTab/SMediaFrameworkCaptureOutputWidget.cpp
+++ b/Engine/Plugins/Media/MediaFrameworkUtilities/Source/MediaFrameworkUtilitiesEditor/Private/CaptureTab/SMediaFrameworkCaptureOutputWidget.cpp
@@ -613,7 +613,7 @@ void SMediaFrameworkCaptureCurrentViewportWidget::StartOutput()
 					LevelViewport = ViewportInterface;
 					SceneViewport = ViewportInterface->GetSharedActiveViewport();
 
-					FLevelEditorViewportClient& ViewportClient = ViewportInterface->GetViewportClient();
+					FEditorViewportClient& ViewportClient = ViewportInterface->GetViewportClient();
 
 					// Save settings
 					ViewportClientFlags.bRealTime = ViewportClient.IsRealtime();
