CIS Fix #rb None

diff --git a/Engine/Plugins/Media/MediaFrameworkUtilities/Source/MediaFrameworkUtilitiesEditor/Private/CaptureTab/SMediaFrameworkCaptureOutputWidget.cpp b/Engine/Plugins/Media/MediaFrameworkUtilities/Source/MediaFrameworkUtilitiesEditor/Private/CaptureTab/SMediaFrameworkCaptureOutputWidget.cpp
--- a/Engine/Plugins/Media/MediaFrameworkUtilities/Source/MediaFrameworkUtilitiesEditor/Private/CaptureTab/SMediaFrameworkCaptureOutputWidget.cpp
+++ b/Engine/Plugins/Media/MediaFrameworkUtilities/Source/MediaFrameworkUtilitiesEditor/Private/CaptureTab/SMediaFrameworkCaptureOutputWidget.cpp
@@ -613,7 +613,7 @@ void SMediaFrameworkCaptureCurrentViewportWidget::StartOutput()
 					LevelViewport = ViewportInterface;
 					SceneViewport = ViewportInterface->GetSharedActiveViewport();
 
-					FEditorViewportClient& ViewportClient = ViewportInterface->GetViewportClient();
+					FEditorViewportClient& ViewportClient = ViewportInterface->GetAssetViewportClient();
 
 					// Save settings
 					ViewportClientFlags.bRealTime = ViewportClient.IsRealtime();
@@ -709,7 +709,7 @@ void SMediaFrameworkCaptureCurrentViewportWidget::ShutdownViewport()
 	TSharedPtr<IAssetViewport> LevelViewportPin = LevelViewport.Pin();
 	if (LevelViewportPin.IsValid() && LevelViewportPin->GetSharedActiveViewport() == EditorSceneViewportPin)
 	{
-		FEditorViewportClient& ViewportClient = LevelViewportPin->GetViewportClient();
+		FEditorViewportClient& ViewportClient = LevelViewportPin->GetAssetViewportClient();
 
 		// Reset settings
 		ViewportClient.SetRealtime(ViewportClientFlags.bRealTime);
