#UE4 Several IWYU fixes

diff --git a/Engine/Plugins/Media/MediaFrameworkUtilities/Source/MediaFrameworkUtilitiesEditor/Private/CaptureTab/SMediaFrameworkCaptureOutputWidget.cpp b/Engine/Plugins/Media/MediaFrameworkUtilities/Source/MediaFrameworkUtilitiesEditor/Private/CaptureTab/SMediaFrameworkCaptureOutputWidget.cpp
--- a/Engine/Plugins/Media/MediaFrameworkUtilities/Source/MediaFrameworkUtilitiesEditor/Private/CaptureTab/SMediaFrameworkCaptureOutputWidget.cpp
+++ b/Engine/Plugins/Media/MediaFrameworkUtilities/Source/MediaFrameworkUtilitiesEditor/Private/CaptureTab/SMediaFrameworkCaptureOutputWidget.cpp
@@ -25,6 +25,7 @@
 #include "Widgets/SViewport.h"
 #include "Widgets/Text/STextBlock.h"
 
+#include "Modules/ModuleManager.h"
 #include "Editor.h"
 #include "IAssetViewport.h"
 #include "LevelEditor.h"
diff --git a/Engine/Plugins/Runtime/nDisplay/Source/DisplayClusterUtils/DisplayClusterCommonHelpers.h b/Engine/Plugins/Runtime/nDisplay/Source/DisplayClusterUtils/DisplayClusterCommonHelpers.h
--- a/Engine/Plugins/Runtime/nDisplay/Source/DisplayClusterUtils/DisplayClusterCommonHelpers.h
+++ b/Engine/Plugins/Runtime/nDisplay/Source/DisplayClusterUtils/DisplayClusterCommonHelpers.h
@@ -10,6 +10,7 @@
 #include "Config/IDisplayClusterConfigManager.h"
 #include "Config/DisplayClusterConfigTypes.h"
 #include "Game/IDisplayClusterGameManager.h"
+#include "DisplayClusterTypesConverter.h"
 
 #include "DisplayClusterUtils/DisplayClusterTypesConverter.h"
 
diff --git a/Engine/Source/Developer/HierarchicalLODUtilities/Private/HierarchicalLODProxyProcessor.cpp b/Engine/Source/Developer/HierarchicalLODUtilities/Private/HierarchicalLODProxyProcessor.cpp
--- a/Engine/Source/Developer/HierarchicalLODUtilities/Private/HierarchicalLODProxyProcessor.cpp
+++ b/Engine/Source/Developer/HierarchicalLODUtilities/Private/HierarchicalLODProxyProcessor.cpp
@@ -8,6 +8,7 @@
 #include "HierarchicalLODUtilitiesModule.h"
 #include "Editor.h"
 
+#include "Modules/ModuleManager.h"
 #include "Interfaces/IProjectManager.h"
 #include "StaticMeshResources.h"
 #include "Logging/TokenizedMessage.h"
diff --git a/Engine/Source/Editor/SceneOutliner/Private/SSceneOutliner.cpp b/Engine/Source/Editor/SceneOutliner/Private/SSceneOutliner.cpp
--- a/Engine/Source/Editor/SceneOutliner/Private/SSceneOutliner.cpp
+++ b/Engine/Source/Editor/SceneOutliner/Private/SSceneOutliner.cpp
@@ -28,6 +28,7 @@
 #include "Textures/SlateIcon.h"
 #include "ToolMenus.h"
 #include "UnrealEdGlobals.h"
+#include "UObject/PackageReload.h"
 #include "Widgets/Images/SImage.h"
 #include "Widgets/Input/SButton.h"
 #include "Widgets/Input/SComboButton.h"
diff --git a/Engine/Source/Editor/UnrealEd/Private/HierarchicalLOD.cpp b/Engine/Source/Editor/UnrealEd/Private/HierarchicalLOD.cpp
--- a/Engine/Source/Editor/UnrealEd/Private/HierarchicalLOD.cpp
+++ b/Engine/Source/Editor/UnrealEd/Private/HierarchicalLOD.cpp
@@ -16,6 +16,7 @@
 
 #include "Algo/Transform.h"
 #include "EditorLevelUtils.h"
+#include "Modules/ModuleManager.h"
 
 #if WITH_EDITOR
 #include "Engine/LODActor.h"
diff --git a/Engine/Source/Editor/UnrealEd/Private/Lightmass/Lightmass.cpp b/Engine/Source/Editor/UnrealEd/Private/Lightmass/Lightmass.cpp
--- a/Engine/Source/Editor/UnrealEd/Private/Lightmass/Lightmass.cpp
+++ b/Engine/Source/Editor/UnrealEd/Private/Lightmass/Lightmass.cpp
@@ -55,6 +55,7 @@
 #include "ComponentRecreateRenderStateContext.h"
 #include "EditorLevelUtils.h"
 #include "Misc/MessageDialog.h"
+#include "Modules/ModuleManager.h"
 
 extern FSwarmDebugOptions GSwarmDebugOptions;
 
diff --git a/Engine/Source/Editor/WorldBrowser/Private/Tiles/WorldTileModel.cpp b/Engine/Source/Editor/WorldBrowser/Private/Tiles/WorldTileModel.cpp
--- a/Engine/Source/Editor/WorldBrowser/Private/Tiles/WorldTileModel.cpp
+++ b/Engine/Source/Editor/WorldBrowser/Private/Tiles/WorldTileModel.cpp
@@ -11,6 +11,7 @@
 #include "EditorLevelUtils.h"
 #include "LevelCollectionModel.h"
 
+#include "Modules/ModuleManager.h"
 #include "Tiles/WorldTileDetails.h"
 #include "Tiles/WorldTileCollectionModel.h"
 #include "Engine/WorldComposition.h"
diff --git a/Engine/Source/Runtime/AudioMixer/Classes/SubmixEffects/AudioMixerSubmixEffectReverb.h b/Engine/Source/Runtime/AudioMixer/Classes/SubmixEffects/AudioMixerSubmixEffectReverb.h
--- a/Engine/Source/Runtime/AudioMixer/Classes/SubmixEffects/AudioMixerSubmixEffectReverb.h
+++ b/Engine/Source/Runtime/AudioMixer/Classes/SubmixEffects/AudioMixerSubmixEffectReverb.h
@@ -5,6 +5,7 @@
 #include "DSP/Reverb.h"
 #include "Sound/SoundEffectSubmix.h"
 #include "AudioEffect.h"
+#include "Curves/RichCurve.h"
 #include "AudioMixerSubmixEffectReverb.generated.h"
 
 USTRUCT(BlueprintType)
diff --git a/Engine/Source/Runtime/Engine/Private/Animation/AnimSequence.cpp b/Engine/Source/Runtime/Engine/Private/Animation/AnimSequence.cpp
--- a/Engine/Source/Runtime/Engine/Private/Animation/AnimSequence.cpp
+++ b/Engine/Source/Runtime/Engine/Private/Animation/AnimSequence.cpp
@@ -40,6 +40,8 @@
 #include "DeviceProfiles/DeviceProfileManager.h"
 #include "DeviceProfiles/DeviceProfile.h"
 #include "Animation/AnimStreamable.h"
+#include "Modules/ModuleManager.h"
+
 #define USE_SLERP 0
 #define LOCTEXT_NAMESPACE "AnimSequence"
 
diff --git a/Engine/Source/Runtime/Engine/Private/ShadowMap.cpp b/Engine/Source/Runtime/Engine/Private/ShadowMap.cpp
--- a/Engine/Source/Runtime/Engine/Private/ShadowMap.cpp
+++ b/Engine/Source/Runtime/Engine/Private/ShadowMap.cpp
@@ -16,6 +16,7 @@
 
 #if WITH_EDITOR
 
+	#include "Modules/ModuleManager.h"
 	#include "TextureCompressorModule.h"
 
 	// NOTE: We're only counting the top-level mip-map for the following variables.
diff --git a/Engine/Source/Runtime/Engine/Private/SkeletalMeshLODRenderData.cpp b/Engine/Source/Runtime/Engine/Private/SkeletalMeshLODRenderData.cpp
--- a/Engine/Source/Runtime/Engine/Private/SkeletalMeshLODRenderData.cpp
+++ b/Engine/Source/Runtime/Engine/Private/SkeletalMeshLODRenderData.cpp
@@ -13,6 +13,7 @@
 #include "UObject/PropertyPortFlags.h"
 
 #if WITH_EDITOR
+#include "Modules/ModuleManager.h"
 #include "Rendering/SkeletalMeshModel.h"
 #include "MeshUtilities.h"
 #endif // WITH_EDITOR
diff --git a/Engine/Source/Runtime/Engine/Private/Timeline.cpp b/Engine/Source/Runtime/Engine/Private/Timeline.cpp
--- a/Engine/Source/Runtime/Engine/Private/Timeline.cpp
+++ b/Engine/Source/Runtime/Engine/Private/Timeline.cpp
@@ -11,6 +11,7 @@
 #include "UObject/UnrealType.h"
 #include "Curves/CurveLinearColor.h"
 #include "Curves/CurveVector.h"
+#include "Curves/CurveFloat.h"
 #include "UObject/Package.h"
 #include "GameFramework/WorldSettings.h"
 #include "Net/UnrealNetwork.h"
