Fix more deprecation warnings

diff --git a/Engine/Plugins/Media/MediaFrameworkUtilities/Source/MediaFrameworkUtilitiesEditor/Private/CaptureTab/SMediaFrameworkCaptureOutputWidget.cpp b/Engine/Plugins/Media/MediaFrameworkUtilities/Source/MediaFrameworkUtilitiesEditor/Private/CaptureTab/SMediaFrameworkCaptureOutputWidget.cpp
--- a/Engine/Plugins/Media/MediaFrameworkUtilities/Source/MediaFrameworkUtilitiesEditor/Private/CaptureTab/SMediaFrameworkCaptureOutputWidget.cpp
+++ b/Engine/Plugins/Media/MediaFrameworkUtilities/Source/MediaFrameworkUtilitiesEditor/Private/CaptureTab/SMediaFrameworkCaptureOutputWidget.cpp
@@ -109,8 +109,7 @@ namespace MediaFrameworkUtilities
 				SetReferenceToWorldContext(GEditor->GetEditorWorldContext());
 			}
 			SetViewMode(ViewModeIndex);
-			SetRealtime(true, false);
-			SetRealtime(true, true); // Save that setting for RestoreRealtime
+			SetRealtime(true);
 		}
 
 		void AutoSetPIE()
@@ -122,8 +121,7 @@ namespace MediaFrameworkUtilities
 			else
 			{
 				SetViewMode(ViewModeIndex);
-				SetRealtime(true, false);
-				SetRealtime(true, true); // Save that setting for RestoreRealtime
+				SetRealtime(true);
 			}
 		}
 
