Removed blank in filename "AssetTypeActions_MediaProfile.cpp" to fix build on Mac.

#rnx

diff --git a/Engine/Plugins/Media/MediaFrameworkUtilities/Source/MediaFrameworkUtilitiesEditor/Private/AssetTypeActions/AssetTypeActions_MediaProfile .cpp b/Engine/Plugins/Media/MediaFrameworkUtilities/Source/MediaFrameworkUtilitiesEditor/Private/AssetTypeActions/AssetTypeActions_MediaProfile .cpp
deleted file mode 100644
diff --git a/Engine/Plugins/Media/MediaFrameworkUtilities/Source/MediaFrameworkUtilitiesEditor/Private/AssetTypeActions/AssetTypeActions_MediaProfile.cpp b/Engine/Plugins/Media/MediaFrameworkUtilities/Source/MediaFrameworkUtilitiesEditor/Private/AssetTypeActions/AssetTypeActions_MediaProfile.cpp
new file mode 100644
--- /dev/null
+++ b/Engine/Plugins/Media/MediaFrameworkUtilities/Source/MediaFrameworkUtilitiesEditor/Private/AssetTypeActions/AssetTypeActions_MediaProfile.cpp
@@ -0,0 +1,31 @@
+// Copyright Epic Games, Inc. All Rights Reserved.
+
+#include "AssetTypeActions/AssetTypeActions_MediaProfile.h"
+
+#include "AssetEditor/MediaProfileEditorToolkit.h"
+#include "Profile/MediaProfile.h"
+
+#define LOCTEXT_NAMESPACE "MediaProfileEditor"
+
+FText FAssetTypeActions_MediaProfile::GetName() const
+{
+	return LOCTEXT("AssetTypeActions_MediaProfile", "Media Profile");
+}
+
+UClass* FAssetTypeActions_MediaProfile::GetSupportedClass() const
+{
+	return UMediaProfile::StaticClass();
+}
+
+void FAssetTypeActions_MediaProfile::OpenAssetEditor(const TArray<UObject*>& InObjects, TSharedPtr<class IToolkitHost> EditWithinLevelEditor)
+{
+	for (UObject* Obj : InObjects)
+	{
+		if (UMediaProfile* Asset = Cast<UMediaProfile>(Obj))
+		{
+			FMediaProfileEditorToolkit::CreateEditor(EToolkitMode::Standalone, EditWithinLevelEditor, Asset);
+		}
+	}
+}
+
+#undef LOCTEXT_NAMESPACE
