Level Editor Menu Clean Up, Part 3/3

diff --git a/Engine/Plugins/Experimental/VirtualCamera/Source/VirtualCameraEditor/Private/VirtualCameraEditorStyle.cpp b/Engine/Plugins/Experimental/VirtualCamera/Source/VirtualCameraEditor/Private/VirtualCameraEditorStyle.cpp
--- a/Engine/Plugins/Experimental/VirtualCamera/Source/VirtualCameraEditor/Private/VirtualCameraEditorStyle.cpp
+++ b/Engine/Plugins/Experimental/VirtualCamera/Source/VirtualCameraEditor/Private/VirtualCameraEditorStyle.cpp
@@ -6,6 +6,10 @@
 #include "Interfaces/IPluginManager.h"
 #include "Styling/SlateStyle.h"
 #include "Styling/SlateStyleRegistry.h"
+#include "Styling/SlateStyleMacros.h"
+
+// This is to fix the issue that SlateStyleMacros like IMAGE_BRUSH look for RootToContentDir but StyleSet->RootToContentDir is how this style is set up
+#define RootToContentDir VirtualCameraStyle::StyleInstance->RootToContentDir
 
 namespace VirtualCameraStyle
 {
@@ -17,8 +21,6 @@ namespace VirtualCameraStyle
 	static TUniquePtr<FSlateStyleSet> StyleInstance;
 }
 
-#define IMAGE_BRUSH(RelativePath, ...) FSlateImageBrush(VirtualCameraStyle::StyleInstance->RootToContentDir(RelativePath, TEXT(".png")), __VA_ARGS__)
-
 void FVirtualCameraEditorStyle::Register()
 {
 	VirtualCameraStyle::StyleInstance = MakeUnique<FSlateStyleSet>(VirtualCameraStyle::NAME_StyleName);
@@ -26,7 +28,7 @@ void FVirtualCameraEditorStyle::Register()
 	const FString ContentRoot = ProjectBaseDir / TEXT("Content/Editor/Icons/");
 	VirtualCameraStyle::StyleInstance->SetContentRoot(ContentRoot);
 
-	VirtualCameraStyle::StyleInstance->Set("TabIcons.VirtualCamera.Small", new IMAGE_BRUSH("VirtualCamera_Stream_16x", VirtualCameraStyle::Icon16x16));
+	VirtualCameraStyle::StyleInstance->Set("TabIcons.VirtualCamera.Small", new IMAGE_BRUSH_SVG("VirtualCamera", VirtualCameraStyle::Icon16x16));
 	VirtualCameraStyle::StyleInstance->Set("VirtualCamera.Stream", new IMAGE_BRUSH("VirtualCamera_Stream_40x", VirtualCameraStyle::Icon40x40));
 	VirtualCameraStyle::StyleInstance->Set("VirtualCamera.Stop", new IMAGE_BRUSH("VirtualCamera_Stop_40x", VirtualCameraStyle::Icon40x40));
 
diff --git a/Engine/Plugins/Experimental/VirtualProductionUtilities/Source/VPUtilitiesEditor/Private/VPUtilitiesEditorStyle.cpp b/Engine/Plugins/Experimental/VirtualProductionUtilities/Source/VPUtilitiesEditor/Private/VPUtilitiesEditorStyle.cpp
--- a/Engine/Plugins/Experimental/VirtualProductionUtilities/Source/VPUtilitiesEditor/Private/VPUtilitiesEditorStyle.cpp
+++ b/Engine/Plugins/Experimental/VirtualProductionUtilities/Source/VPUtilitiesEditor/Private/VPUtilitiesEditorStyle.cpp
@@ -24,7 +24,7 @@ void FVPUtilitiesEditorStyle::Register()
 	VPUtilitiesEditorStyle::StyleInstance = MakeUnique<FSlateStyleSet>(VPUtilitiesEditorStyle::NAME_StyleName);
 	VPUtilitiesEditorStyle::StyleInstance->SetContentRoot(FPaths::EnginePluginsDir() / TEXT("Experimental/VirtualProductionUtilities/Content/Editor/Icons/"));
 
-	VPUtilitiesEditorStyle::StyleInstance->Set("TabIcons.Genlock.Small", new IMAGE_BRUSH("Icon_GenlockTab_16x", Icon16x16));
+	VPUtilitiesEditorStyle::StyleInstance->Set("TabIcons.Genlock.Small", new IMAGE_BRUSH_SVG("Genlock", Icon16x16));
 
 	VPUtilitiesEditorStyle::StyleInstance->Set("PlacementBrowser.Icons.VirtualProduction", new IMAGE_BRUSH_SVG("VirtualProduction", Icon16x16));
 
diff --git a/Engine/Plugins/Media/MediaFrameworkUtilities/Source/MediaFrameworkUtilitiesEditor/Private/UI/MediaFrameworkUtilitiesEditorStyle.cpp b/Engine/Plugins/Media/MediaFrameworkUtilities/Source/MediaFrameworkUtilitiesEditor/Private/UI/MediaFrameworkUtilitiesEditorStyle.cpp
--- a/Engine/Plugins/Media/MediaFrameworkUtilities/Source/MediaFrameworkUtilitiesEditor/Private/UI/MediaFrameworkUtilitiesEditorStyle.cpp
+++ b/Engine/Plugins/Media/MediaFrameworkUtilities/Source/MediaFrameworkUtilitiesEditor/Private/UI/MediaFrameworkUtilitiesEditorStyle.cpp
@@ -5,6 +5,10 @@
 #include "Brushes/SlateImageBrush.h"
 #include "Styling/SlateStyle.h"
 #include "Styling/SlateStyleRegistry.h"
+#include "Styling/SlateStyleMacros.h"
+
+// This is to fix the issue that SlateStyleMacros like IMAGE_BRUSH look for RootToContentDir but StyleSet->RootToContentDir is how this style is set up
+#define RootToContentDir MediaFrameworkUtilitiesStyle::StyleInstance->RootToContentDir
 
 namespace MediaFrameworkUtilitiesStyle
 {
@@ -19,8 +23,6 @@ namespace MediaFrameworkUtilitiesStyle
 	static TUniquePtr<FSlateStyleSet> StyleInstance;
 }
 
-#define IMAGE_BRUSH(RelativePath, ...) FSlateImageBrush(MediaFrameworkUtilitiesStyle::StyleInstance->RootToContentDir(RelativePath, TEXT(".png")), __VA_ARGS__)
-
 void FMediaFrameworkUtilitiesEditorStyle::Register()
 {
 	MediaFrameworkUtilitiesStyle::StyleInstance = MakeUnique<FSlateStyleSet>(MediaFrameworkUtilitiesStyle::NAME_StyleName);
@@ -38,12 +40,12 @@ void FMediaFrameworkUtilitiesEditorStyle::Register()
 	MediaFrameworkUtilitiesStyle::StyleInstance->Set("ToolbarIcon.MediaProfile", new IMAGE_BRUSH("MediaProfile_Color_40x", MediaFrameworkUtilitiesStyle::Icon40x40));
 	MediaFrameworkUtilitiesStyle::StyleInstance->Set("ToolbarIcon.Apply", new IMAGE_BRUSH("Apply_40x", MediaFrameworkUtilitiesStyle::Icon40x40));
 
-	MediaFrameworkUtilitiesStyle::StyleInstance->Set("TabIcons.MediaCapture.Small", new IMAGE_BRUSH("CaptureCameraViewport_Capture_16x", MediaFrameworkUtilitiesStyle::Icon16x16));
+	MediaFrameworkUtilitiesStyle::StyleInstance->Set("TabIcons.MediaCapture.Small", new IMAGE_BRUSH_SVG("MediaCapture", MediaFrameworkUtilitiesStyle::Icon16x16));
 	MediaFrameworkUtilitiesStyle::StyleInstance->Set("MediaCapture.Capture", new IMAGE_BRUSH("CaptureCameraViewport_Capture_40x", MediaFrameworkUtilitiesStyle::Icon40x40));
 	MediaFrameworkUtilitiesStyle::StyleInstance->Set("MediaCapture.Stop", new IMAGE_BRUSH("CaptureCameraViewport_Stop_40x", MediaFrameworkUtilitiesStyle::Icon40x40));
 	MediaFrameworkUtilitiesStyle::StyleInstance->Set("MediaCapture.Settings", new IMAGE_BRUSH("Settings_40x", MediaFrameworkUtilitiesStyle::Icon40x40));
 
-	MediaFrameworkUtilitiesStyle::StyleInstance->Set("TabIcons.VideoInput.Small", new IMAGE_BRUSH("Icon_VideoInputTab_16x", MediaFrameworkUtilitiesStyle::Icon16x16));
+	MediaFrameworkUtilitiesStyle::StyleInstance->Set("TabIcons.VideoInput.Small", new IMAGE_BRUSH_SVG("MediaPlayback", MediaFrameworkUtilitiesStyle::Icon16x16));
 	MediaFrameworkUtilitiesStyle::StyleInstance->Set("VideoInput.Play", new IMAGE_BRUSH("MediaSource_Play_40x", MediaFrameworkUtilitiesStyle::Icon40x40));
 	MediaFrameworkUtilitiesStyle::StyleInstance->Set("VideoInput.Stop", new IMAGE_BRUSH("CaptureCameraViewport_Stop_40x", MediaFrameworkUtilitiesStyle::Icon40x40));
 	MediaFrameworkUtilitiesStyle::StyleInstance->Set("VideoInput.Settings", new IMAGE_BRUSH("Settings_40x", MediaFrameworkUtilitiesStyle::Icon40x40));
diff --git a/Engine/Plugins/Tests/FunctionalTestingEditor/Source/FunctionalTestingEditor.Build.cs b/Engine/Plugins/Tests/FunctionalTestingEditor/Source/FunctionalTestingEditor.Build.cs
--- a/Engine/Plugins/Tests/FunctionalTestingEditor/Source/FunctionalTestingEditor.Build.cs
+++ b/Engine/Plugins/Tests/FunctionalTestingEditor/Source/FunctionalTestingEditor.Build.cs
@@ -31,6 +31,7 @@ public class FunctionalTestingEditor : ModuleRules
                 "PlacementMode",
                 "WorkspaceMenuStructure",
                 "ScreenShotComparisonTools",
+                "ToolMenus",
 				"AssetTools"
             }
          );
diff --git a/Engine/Plugins/Tests/FunctionalTestingEditor/Source/Private/FunctionalTestingEditorModule.cpp b/Engine/Plugins/Tests/FunctionalTestingEditor/Source/Private/FunctionalTestingEditorModule.cpp
--- a/Engine/Plugins/Tests/FunctionalTestingEditor/Source/Private/FunctionalTestingEditorModule.cpp
+++ b/Engine/Plugins/Tests/FunctionalTestingEditor/Source/Private/FunctionalTestingEditorModule.cpp
@@ -24,6 +24,7 @@
 #include "Widgets/Docking/SDockTab.h"
 #include "AssetTypeActions_GroundTruthData.h"
 #include "IAssetTools.h"
+#include "ToolMenus.h"
 
 #define LOCTEXT_NAMESPACE "EditorAutomation"
 
@@ -105,17 +106,16 @@ class FFunctionalTestingEditorModule : public IFunctionalTestingEditorModule
 {
 	void StartupModule()
 	{
-		// make an extension to add the Orion function menu
-		Extender = MakeShareable(new FExtender());
-		Extender->AddMenuExtension(
-			"General",
-			EExtensionHook::After,
-			nullptr,
-			FMenuExtensionDelegate::CreateRaw(this, &FFunctionalTestingEditorModule::OnAutomationToolsMenuCreation));
-
-		// add the menu extension to the editor
-		FLevelEditorModule& LevelEditorModule = FModuleManager::LoadModuleChecked<FLevelEditorModule>("LevelEditor");
-		LevelEditorModule.GetMenuExtensibilityManager()->AddExtender(Extender);
+		// Add Automation Area to the Tools Menu
+		UToolMenu* Menu = UToolMenus::Get()->ExtendMenu("LevelEditor.MainMenu.Tools");
+		FToolMenuSection& Section = Menu->AddSection("Automation", LOCTEXT("AutomationHeading", "Automation"));
+		Section.AddMenuEntry( 
+			"TestAutomation",
+			LOCTEXT("AutomationLabel", "Test Automation"),
+			LOCTEXT("Tooltip", "Launch the Testing Automation Frontend."),
+			FSlateIcon(FAppStyle::GetAppStyleSetName(), "AutomationTools.TestAutomation"),
+			FUIAction(FExecuteAction::CreateStatic(&FFunctionalTestingEditorModule::OnShowAutomationFrontend))
+		);
 
 		FModuleManager::Get().OnModulesChanged().AddRaw(this, &FFunctionalTestingEditorModule::OnModulesChanged);
 
@@ -144,30 +144,6 @@ class FFunctionalTestingEditorModule : public IFunctionalTestingEditorModule
 		}
 	}
 
-	void OnAutomationToolsMenuCreation(FMenuBuilder& MenuBuilder)
-	{
-		MenuBuilder.BeginSection("Testing", LOCTEXT("Testing", "Testing"));
-		MenuBuilder.AddMenuEntry(
-			LOCTEXT("AutomationLabel", "Test Automation"),
-			LOCTEXT("Tooltip", "Launch the Testing Automation Frontend."),
-			FSlateIcon(FAppStyle::GetAppStyleSetName(), "AutomationTools.TestAutomation"),
-			FUIAction(FExecuteAction::CreateStatic(&FFunctionalTestingEditorModule::OnShowAutomationFrontend)));
-		MenuBuilder.EndSection();
-
-		// TODO Come up with a way to hide this if no tools are registered.
-		if (WorkspaceMenu::GetMenuStructure().GetAutomationToolsCategory()->GetChildItems().Num() > 0)
-		{
-			MenuBuilder.AddSubMenu(
-				LOCTEXT("AutomationTools", "Automation Tools"),
-				LOCTEXT("AutomationToolsToolTip", "Assorted tools to help generate data for some of the automation tests."),
-				FNewMenuDelegate::CreateRaw(this, &FFunctionalTestingEditorModule::PopulateAutomationTools),
-				false,
-				FSlateIcon(FAppStyle::GetAppStyleSetName(), "AutomationTools.MenuIcon")
-
-			);
-		}
-	}
-
 	void PopulateAutomationTools(FMenuBuilder& MenuBuilder)
 	{
 		MenuBuilder.BeginSection("AutomationTools", LOCTEXT("AutomationTools", "Automation Tools"));
diff --git a/Engine/Source/Editor/EditorStyle/Private/StarshipStyle.cpp b/Engine/Source/Editor/EditorStyle/Private/StarshipStyle.cpp
--- a/Engine/Source/Editor/EditorStyle/Private/StarshipStyle.cpp
+++ b/Engine/Source/Editor/EditorStyle/Private/StarshipStyle.cpp
@@ -4379,8 +4379,6 @@ void FStarshipEditorStyle::FStyle::SetupLevelEditorStyle()
 			Set( "LevelEditor.Tabs.ComposureCompositing",   new IMAGE_BRUSH_SVG("Starship/Common/ComposureCompositing", Icon16x16) );
 			Set( "LevelEditor.Tabs.USDStage",   			new IMAGE_BRUSH_SVG("Starship/Common/USDStage", Icon16x16) );
 
-			Set( "LevelEditor.Tabs.BuildAndSubmit",         new IMAGE_BRUSH("Icons/icon_tab_BuildSubmit_16x", Icon16x16 ) );
-
 			Set( "LevelEditor.Tabs.StatsViewer",            new IMAGE_BRUSH_SVG("Starship/Common/Statistics", Icon16x16) );
 
 			Set( "LevelEditor.Tabs.Toolbar",                new IMAGE_BRUSH("Icons/icon_tab_Toolbars_16x", Icon16x16 ) );
@@ -4406,11 +4404,6 @@ void FStarshipEditorStyle::FStyle::SetupLevelEditorStyle()
 		Set( "LevelEditor.ExportAll",      new IMAGE_BRUSH_SVG( "Starship/Common/ExportAll", Icon16x16 ) );
 		Set( "LevelEditor.ExportSelected", new IMAGE_BRUSH_SVG( "Starship/Common/ExportSelected", Icon16x16 ) );
 
-		Set( "LevelEditor.Build", new IMAGE_BRUSH( "Icons/icon_build_40x", Icon40x40 ) );
-		Set( "LevelEditor.Build.Small", new IMAGE_BRUSH( "Icons/icon_build_40x", Icon20x20 ) );
-		Set( "LevelEditor.MapCheck", new IMAGE_BRUSH( "Icons/icon_MapCheck_40x", Icon40x40 ) );
-		Set( "LevelEditor.MapCheck.Small", new IMAGE_BRUSH( "Icons/icon_MapCheck_40x", Icon20x20 ) );
-
 		Set( "LevelEditor.Recompile", new IMAGE_BRUSH_SVG( "Starship/MainToolbar/compile", Icon40x40 ) );
 		Set( "LevelEditor.Recompile.Small", new IMAGE_BRUSH_SVG( "Starship/MainToolbar/compile", Icon20x20 ) );
 
diff --git a/Engine/Source/Editor/LevelInstanceEditor/Private/LevelInstanceEditorModule.cpp b/Engine/Source/Editor/LevelInstanceEditor/Private/LevelInstanceEditorModule.cpp
--- a/Engine/Source/Editor/LevelInstanceEditor/Private/LevelInstanceEditorModule.cpp
+++ b/Engine/Source/Editor/LevelInstanceEditor/Private/LevelInstanceEditorModule.cpp
@@ -794,7 +794,7 @@ void FLevelInstanceEditorModule::ExtendContextMenu()
 			}));
 
 		FToolMenuEntry& Entry = Section.AddMenuEntry(NAME_None, LOCTEXT("PackLevelInstancesTitle", "Pack Level Instances"),
-			LOCTEXT("PackLevelInstancesTooltip", "Update packed level instances and blueprints"), FSlateIcon(FEditorStyle::GetStyleSetName(), "ClassIcon.LevelInstance"), PackAction, EUserInterfaceActionType::Button);
+			LOCTEXT("PackLevelInstancesTooltip", "Update packed level instances and blueprints"), FSlateIcon(), PackAction, EUserInterfaceActionType::Button);
 	}
 
 	if (UToolMenu* Menu = UToolMenus::Get()->ExtendMenu("LevelEditor.ActorContextMenu"))
