Adding placeholder media icon

diff --git a/Engine/Plugins/Media/MediaFrameworkUtilities/Source/MediaFrameworkUtilitiesEditor/Private/MediaFrameworkUtilitiesPlacement.cpp b/Engine/Plugins/Media/MediaFrameworkUtilities/Source/MediaFrameworkUtilitiesEditor/Private/MediaFrameworkUtilitiesPlacement.cpp
--- a/Engine/Plugins/Media/MediaFrameworkUtilities/Source/MediaFrameworkUtilitiesEditor/Private/MediaFrameworkUtilitiesPlacement.cpp
+++ b/Engine/Plugins/Media/MediaFrameworkUtilities/Source/MediaFrameworkUtilitiesEditor/Private/MediaFrameworkUtilitiesPlacement.cpp
@@ -210,7 +210,7 @@ void FMediaFrameworkUtilitiesPlacement::RegisterPlacement()
 {
 	IPlacementModeModule& PlacementModeModule = IPlacementModeModule::Get();
 	FName CategoryName = "Media";
-	FPlacementCategoryInfo CategoryInfo(LOCTEXT("PlacementMode_Media", "Media"), CategoryName, TEXT("PMMedia"), 35);
+	FPlacementCategoryInfo CategoryInfo(LOCTEXT("PlacementMode_Media", "Media"), FSlateIcon(FAppStyle::Get().GetStyleSetName(), TEXT("SequenceRecorder.TabIcon")), CategoryName, TEXT("PMMedia"), 35);
 	CategoryInfo.CustomGenerator = []() -> TSharedRef<SWidget> { return SNew(SMediaPlacementPalette); };
 	PlacementModeModule.RegisterPlacementCategory(CategoryInfo);
 }
