- Fixing MediaProfile and LensFile menus

diff --git a/Engine/Plugins/Media/MediaFrameworkUtilities/Source/MediaFrameworkUtilitiesEditor/MediaFrameworkUtilitiesEditor.Build.cs b/Engine/Plugins/Media/MediaFrameworkUtilities/Source/MediaFrameworkUtilitiesEditor/MediaFrameworkUtilitiesEditor.Build.cs
--- a/Engine/Plugins/Media/MediaFrameworkUtilities/Source/MediaFrameworkUtilitiesEditor/MediaFrameworkUtilitiesEditor.Build.cs
+++ b/Engine/Plugins/Media/MediaFrameworkUtilities/Source/MediaFrameworkUtilitiesEditor/MediaFrameworkUtilitiesEditor.Build.cs
@@ -32,6 +32,7 @@ namespace UnrealBuildTool.Rules
 					"Slate",
 					"SlateCore",
 					"TimeManagement",
+					"ToolMenus",
 					"UnrealEd",
 					"WorkspaceMenuStructure",
 				});
diff --git a/Engine/Plugins/Media/MediaFrameworkUtilities/Source/MediaFrameworkUtilitiesEditor/Private/UI/MediaProfileMenuEntry.cpp b/Engine/Plugins/Media/MediaFrameworkUtilities/Source/MediaFrameworkUtilitiesEditor/Private/UI/MediaProfileMenuEntry.cpp
--- a/Engine/Plugins/Media/MediaFrameworkUtilities/Source/MediaFrameworkUtilitiesEditor/Private/UI/MediaProfileMenuEntry.cpp
+++ b/Engine/Plugins/Media/MediaFrameworkUtilities/Source/MediaFrameworkUtilitiesEditor/Private/UI/MediaProfileMenuEntry.cpp
@@ -9,16 +9,16 @@
 #include "Framework/Application/SlateApplication.h"
 #include "Framework/Commands/UIAction.h"
 #include "Framework/MultiBox/MultiBoxBuilder.h"
-#include "Framework/MultiBox/MultiBoxExtender.h"
 #include "IAssetTools.h"
 #include "Kismet2/SClassPickerDialog.h"
 #include "LevelEditor.h"
 #include "Misc/FeedbackContext.h"
-#include "Modules/ModuleManager.h"
 #include "Profile/IMediaProfileManager.h"
 #include "Profile/MediaProfile.h"
 #include "Profile/MediaProfileSettings.h"
 #include "PropertyCustomizationHelpers.h"
+#include "ToolMenuEntry.h"
+#include "ToolMenus.h"
 
 #include "UI/MediaFrameworkUtilitiesEditorStyle.h"
 #include "Subsystems/AssetEditorSubsystem.h"
@@ -29,7 +29,7 @@ struct FMediaProfileMenuEntryImpl
 {
 	FMediaProfileMenuEntryImpl()
 	{
-		TSharedPtr<FUICommandList> Actions = MakeShareable(new FUICommandList);
+		TSharedPtr<FUICommandList> Actions = MakeShared<FUICommandList>();
 
 		// Action to edit the current selected media profile
 		Actions->MapAction(FMediaProfileCommands::Get().Edit,
@@ -40,26 +40,52 @@ struct FMediaProfileMenuEntryImpl
 					GEditor->GetEditorSubsystem<UAssetEditorSubsystem>()->OpenEditorForAsset(MediaProfile);
 				}
 			}),
-			FCanExecuteAction::CreateLambda([] { return IMediaProfileManager::Get().GetCurrentMediaProfile() != nullptr; }),
-			FIsActionChecked::CreateLambda([]{ return IMediaProfileManager::Get().GetCurrentMediaProfile() != nullptr; })
+			FCanExecuteAction::CreateLambda([] { return IMediaProfileManager::Get().GetCurrentMediaProfile() != nullptr; })
 		);
 
-		ToolBarExtender = MakeShareable(new FExtender);
-		ToolBarExtender->AddToolBarExtension("Settings", EExtensionHook::After, Actions, FToolBarExtensionDelegate::CreateRaw(this, &FMediaProfileMenuEntryImpl::FillToolbar));
+		UToolMenu* Menu = UToolMenus::Get()->ExtendMenu("LevelEditor.LevelEditorToolBar.ModesToolBar");
+		FToolMenuSection& Section = Menu->FindOrAddSection("MediaProfile");
 
-		FLevelEditorModule& LevelEditorModule = FModuleManager::LoadModuleChecked<FLevelEditorModule>("LevelEditor");
-		LevelEditorModule.GetToolBarExtensibilityManager()->AddExtender(ToolBarExtender);
+		auto ButtonTooltipLambda = [this]()
+		{
+			UMediaProfile* MediaProfile = IMediaProfileManager::Get().GetCurrentMediaProfile();
+			if (MediaProfile == nullptr)
+			{
+				return LOCTEXT("EmptyMediaProfile_ToolTip", "Select a Media Profile to edit it.");
+			}
+			return FText::Format(LOCTEXT("MediaProfile_ToolTip", "Edit '{0}'")
+				, FText::FromName(MediaProfile->GetFName()));
+		};
+
+		// Add a button to edit the current media profile
+		FToolMenuEntry MediaProfileButtonEntry = FToolMenuEntry::InitToolBarButton(
+			FMediaProfileCommands::Get().Edit,
+			LOCTEXT("MediaProfile_Label", "Media Profile"),
+			MakeAttributeLambda(ButtonTooltipLambda),
+			FSlateIcon(FMediaFrameworkUtilitiesEditorStyle::GetStyleSetName(), TEXT("ToolbarIcon.MediaProfile"))
+		);
+		MediaProfileButtonEntry.SetCommandList(Actions);
+
+		FToolMenuEntry MediaProfileMenuEntry = FToolMenuEntry::InitComboButton(
+		"MediaProfileMenu",
+		FUIAction(),
+		FOnGetContent::CreateRaw(this, &FMediaProfileMenuEntryImpl::GenerateMenuContent),
+		LOCTEXT("LevelEditorToolbarMediaProfileButtonLabel", "MediaProfile"),
+		LOCTEXT("LevelEditorToolbarMediaProfileButtonTooltip", "Configure current MediaProfile"),
+		FSlateIcon(),
+		true //bInSimpleComboBox
+		);
+
+		MediaProfileMenuEntry.StyleNameOverride = "CalloutToolbar";
+		Section.AddEntry(MediaProfileButtonEntry);
+		Section.AddEntry(MediaProfileMenuEntry);
 	}
 
 	~FMediaProfileMenuEntryImpl()
 	{
-		if (!IsEngineExitRequested() && ToolBarExtender.IsValid())
+		if (!IsEngineExitRequested())
 		{
-			FLevelEditorModule* LevelEditorModule = FModuleManager::GetModulePtr<FLevelEditorModule>("LevelEditor");
-			if (LevelEditorModule)
-			{
-				LevelEditorModule->GetToolBarExtensibilityManager()->RemoveExtender(ToolBarExtender);
-			}
+			UToolMenus::Get()->RemoveSection("LevelEditor.LevelEditorToolBar.ModesToolBar", "MediaProfile");
 		}
 	}
 
@@ -116,43 +142,6 @@ struct FMediaProfileMenuEntryImpl
 		IMediaProfileManager::Get().SetCurrentMediaProfile(Asset);
 	}
 
-	void FillToolbar(FToolBarBuilder& ToolbarBuilder)
-	{
-		ToolbarBuilder.BeginSection("Media Profile");
-		{
-			auto TooltipLambda = [this]()
-			{
-				UMediaProfile* MediaProfile = IMediaProfileManager::Get().GetCurrentMediaProfile();
-				if (MediaProfile == nullptr)
-				{
-					return LOCTEXT("EmptyMediaProfile_ToolTip", "Select a Media Profile to edit it.");
-				}
-				return FText::Format(LOCTEXT("MediaProfile_ToolTip", "Edit '{0}'")
-					, FText::FromName(MediaProfile->GetFName()));
-			};
-
-			// Add a button to edit the current media profile
-			ToolbarBuilder.AddToolBarButton(
-				FMediaProfileCommands::Get().Edit,
-				NAME_None,
-				LOCTEXT("MediaProfile_Label", "Media Profile"),
-				MakeAttributeLambda(TooltipLambda),
-				FSlateIcon(FMediaFrameworkUtilitiesEditorStyle::GetStyleSetName(), TEXT("ToolbarIcon.MediaProfile"))
-			);
-
-			// Add a simple drop-down menu (no label, no icon for the drop-down button itself) that list the media profile available
-			ToolbarBuilder.AddComboButton(
-				FUIAction(),
-				FOnGetContent::CreateRaw(this, &FMediaProfileMenuEntryImpl::GenerateMenuContent),
-				FText::GetEmpty(),
-				LOCTEXT("MediaProfileButton_ToolTip", "List of Media Profile available to the user for editing or creation."),
-				FSlateIcon(),
-				true
-			);
-		}
-		ToolbarBuilder.EndSection();
-	}
-
 	TSharedRef<SWidget> GenerateMenuContent()
 	{
 		const bool bShouldCloseWindowAfterMenuSelection = true;
@@ -215,10 +204,6 @@ struct FMediaProfileMenuEntryImpl
 		);
 	}
 
-private:
-
-	TSharedPtr<FExtender> ToolBarExtender;
-
 public:
 
 	static TUniquePtr<FMediaProfileMenuEntryImpl> Implementation;
diff --git a/Engine/Plugins/VirtualProduction/CameraCalibration/Source/CameraCalibrationEditor/CameraCalibrationEditor.Build.cs b/Engine/Plugins/VirtualProduction/CameraCalibration/Source/CameraCalibrationEditor/CameraCalibrationEditor.Build.cs
--- a/Engine/Plugins/VirtualProduction/CameraCalibration/Source/CameraCalibrationEditor/CameraCalibrationEditor.Build.cs
+++ b/Engine/Plugins/VirtualProduction/CameraCalibration/Source/CameraCalibrationEditor/CameraCalibrationEditor.Build.cs
@@ -39,6 +39,7 @@ namespace UnrealBuildTool.Rules
 					"Settings",
 					"Slate",
 					"SlateCore",
+					"ToolMenus",
 					"UnrealEd",
 					"WorkspaceMenuStructure",
 				});
diff --git a/Engine/Plugins/VirtualProduction/CameraCalibration/Source/CameraCalibrationEditor/Private/UI/CameraCalibrationMenuEntry.cpp b/Engine/Plugins/VirtualProduction/CameraCalibration/Source/CameraCalibrationEditor/Private/UI/CameraCalibrationMenuEntry.cpp
--- a/Engine/Plugins/VirtualProduction/CameraCalibration/Source/CameraCalibrationEditor/Private/UI/CameraCalibrationMenuEntry.cpp
+++ b/Engine/Plugins/VirtualProduction/CameraCalibration/Source/CameraCalibrationEditor/Private/UI/CameraCalibrationMenuEntry.cpp
@@ -4,24 +4,22 @@
 
 #include "AssetEditor/CameraCalibrationCommands.h"
 #include "AssetToolsModule.h"
+#include "CameraCalibrationSettings.h"
+#include "CameraCalibrationSubsystem.h"
 #include "Editor.h"
 #include "Engine/Engine.h"
 #include "Factories/LensFileFactoryNew.h"
 #include "Framework/Application/SlateApplication.h"
 #include "Framework/Commands/UICommandList.h"
 #include "Framework/MultiBox/MultiBoxBuilder.h"
-#include "Framework/MultiBox/MultiBoxExtender.h"
 #include "IAssetTools.h"
 #include "Misc/FeedbackContext.h"
-#include "CameraCalibrationSettings.h"
-#include "CameraCalibrationSubsystem.h"
 #include "LensFile.h"
-#include "LevelEditor.h"
-#include "Modules/ModuleManager.h"
 #include "PropertyCustomizationHelpers.h"
-#include "UI/CameraCalibrationEditorStyle.h"
 #include "Subsystems/AssetEditorSubsystem.h"
 #include "Textures/SlateIcon.h"
+#include "ToolMenus.h"
+#include "UI/CameraCalibrationEditorStyle.h"
 
 #define LOCTEXT_NAMESPACE "CameraCalibrationMenu"
 
@@ -41,31 +39,53 @@ struct FCameraCalibrationMenuEntryImpl
 				}
 			}),
 			FCanExecuteAction::CreateLambda([this] 
-			{ 
-				return GetDefaultLensFile() != nullptr;
-			}),
-			FIsActionChecked::CreateLambda([this]
 			{ 
 				return GetDefaultLensFile() != nullptr;
 			})
 		);
 
-		ToolBarExtender = MakeShared<FExtender>();
-		ToolBarExtender->AddToolBarExtension("Settings", EExtensionHook::After, Actions, FToolBarExtensionDelegate::CreateRaw(this, &FCameraCalibrationMenuEntryImpl::FillToolbar));
+		UToolMenu* Menu = UToolMenus::Get()->ExtendMenu("LevelEditor.LevelEditorToolBar.ModesToolBar");
+		FToolMenuSection& Section = Menu->FindOrAddSection("LensFile");
 
-		FLevelEditorModule& LevelEditorModule = FModuleManager::LoadModuleChecked<FLevelEditorModule>("LevelEditor");
-		LevelEditorModule.GetToolBarExtensibilityManager()->AddExtender(ToolBarExtender);
+		auto ButtonTooltipLambda = [this]()
+		{
+			ULensFile* LensFile = GetDefaultLensFile();
+			if (LensFile == nullptr)
+			{
+				return LOCTEXT("NoFile_ToolTip", "Select a Lens File to edit it.");
+			}
+			return FText::Format(LOCTEXT("LensFile_ToolTip", "Edit '{0}'") , FText::FromName(LensFile->GetFName()));
+		};
+
+		// Add a button to edit the default LensFile
+		FToolMenuEntry LensFileButtonEntry = FToolMenuEntry::InitToolBarButton(
+			FCameraCalibrationCommands::Get().Edit,
+			LOCTEXT("LensFile_Label", "Lens File"),
+			MakeAttributeLambda(ButtonTooltipLambda),
+			FSlateIcon(FCameraCalibrationEditorStyle::Get().GetStyleSetName(), TEXT("ClassIcon.LensFile"))
+		);
+		LensFileButtonEntry.SetCommandList(Actions);
+
+		FToolMenuEntry LensFileMenuEntry = FToolMenuEntry::InitComboButton(
+		"LensFileMenu",
+		FUIAction(),
+		FOnGetContent::CreateRaw(this, &FCameraCalibrationMenuEntryImpl::GenerateMenuContent),
+		LOCTEXT("LevelEditorToolbarLensFileButtonLabel", "LensFile"),
+		LOCTEXT("LevelEditorToolbarLensFileButtonTooltip", "Configure default LensFile"),
+		FSlateIcon(),
+		true //bInSimpleComboBox
+		);
+
+		LensFileMenuEntry.StyleNameOverride = "CalloutToolbar";
+		Section.AddEntry(LensFileButtonEntry);
+		Section.AddEntry(LensFileMenuEntry);
 	}
 
 	~FCameraCalibrationMenuEntryImpl()
 	{
-		if (!IsEngineExitRequested() && ToolBarExtender.IsValid())
+		if (!IsEngineExitRequested())
 		{
-			FLevelEditorModule* LevelEditorModule = FModuleManager::GetModulePtr<FLevelEditorModule>("LevelEditor");
-			if (LevelEditorModule)
-			{
-				LevelEditorModule->GetToolBarExtensibilityManager()->RemoveExtender(ToolBarExtender);
-			}
+			UToolMenus::Get()->RemoveSection("LevelEditor.LevelEditorToolBar.ModesToolBar", "LensFile");
 		}
 	}
 
@@ -108,42 +128,6 @@ struct FCameraCalibrationMenuEntryImpl
 		return SubSystem->SetDefaultLensFile(Asset);
 	}
 
-	void FillToolbar(FToolBarBuilder& ToolbarBuilder)
-	{
-		ToolbarBuilder.BeginSection("Lens Distortion");
-		{
-			auto TooltipLambda = [this]()
-			{
-				ULensFile* LensFile = GetDefaultLensFile();
-				if (LensFile == nullptr)
-				{
-					return LOCTEXT("NoFile_ToolTip", "Select a Lens File to edit it.");
-				}
-				return FText::Format(LOCTEXT("LensFile_ToolTip", "Edit '{0}'") , FText::FromName(LensFile->GetFName()));
-			};
-
-			// Add a button to edit the current lens file
-			ToolbarBuilder.AddToolBarButton(
-				FCameraCalibrationCommands::Get().Edit,
-				NAME_None,
-				LOCTEXT("LensFile_Label", "Lens File"),
-				MakeAttributeLambda(TooltipLambda),
-				FSlateIcon(FCameraCalibrationEditorStyle::Get().GetStyleSetName(), TEXT("ToolbarIcon.LensFile"))
-			);
-
-			// Add a simple drop-down menu (no label, no icon for the drop-down button itself) that list the lens files available
-			ToolbarBuilder.AddComboButton(
-				FUIAction(),
-				FOnGetContent::CreateRaw(this, &FCameraCalibrationMenuEntryImpl::GenerateMenuContent),
-				FText::GetEmpty(),
-				LOCTEXT("LensFileButton_ToolTip", "List of Lens Files available to the user for editing."),
-				FSlateIcon(),
-				true
-			);
-		}
-		ToolbarBuilder.EndSection();
-	}
-
 	TSharedRef<SWidget> GenerateMenuContent()
 	{
 		const bool bShouldCloseWindowAfterMenuSelection = true;
@@ -206,10 +190,6 @@ struct FCameraCalibrationMenuEntryImpl
 		);
 	}
 
-private:
-
-	TSharedPtr<FExtender> ToolBarExtender;
-
 public:
 
 	static TUniquePtr<FCameraCalibrationMenuEntryImpl> Implementation;
