AssetRegistry includes (Engine Plugins): change #include "AssetData.h" -> #include "AssetRegistry/AssetData.h", and similar for the other moved AssetRegistry headers.
#rnx

diff --git a/Engine/Plugins/Animation/ControlRig/Source/ControlRig/Public/Tools/ControlRigToolAsset.h b/Engine/Plugins/Animation/ControlRig/Source/ControlRig/Public/Tools/ControlRigToolAsset.h
--- a/Engine/Plugins/Animation/ControlRig/Source/ControlRig/Public/Tools/ControlRigToolAsset.h
+++ b/Engine/Plugins/Animation/ControlRig/Source/ControlRig/Public/Tools/ControlRigToolAsset.h
@@ -5,8 +5,8 @@
 #include "UObject/Object.h"
 #include "ControlRig.h"
 #include "Tools/ControlRigPoseProjectSettings.h"
-#include "IAssetRegistry.h"
-#include "AssetRegistryModule.h"
+#include "AssetRegistry/IAssetRegistry.h"
+#include "AssetRegistry/AssetRegistryModule.h"
 #include "Widgets/Notifications/SNotificationList.h"
 #include "Framework/Notifications/NotificationManager.h"
 #include "ControlRigObjectBinding.h"
diff --git a/Engine/Plugins/Developer/PluginUtils/Source/PluginUtils/Private/PluginUtils.cpp b/Engine/Plugins/Developer/PluginUtils/Source/PluginUtils/Private/PluginUtils.cpp
--- a/Engine/Plugins/Developer/PluginUtils/Source/PluginUtils/Private/PluginUtils.cpp
+++ b/Engine/Plugins/Developer/PluginUtils/Source/PluginUtils/Private/PluginUtils.cpp
@@ -10,9 +10,9 @@
 #include "Interfaces/IProjectManager.h"
 #include "Interfaces/IPluginManager.h"
 #include "PluginDescriptor.h"
-#include "IAssetRegistry.h"
-#include "AssetRegistryModule.h"
-#include "AssetData.h"
+#include "AssetRegistry/IAssetRegistry.h"
+#include "AssetRegistry/AssetRegistryModule.h"
+#include "AssetRegistry/AssetData.h"
 #include "IAssetTools.h"
 #include "AssetToolsModule.h"
 #include "DesktopPlatformModule.h"
diff --git a/Engine/Plugins/Developer/TraceSourceFiltering/Source/SourceFilteringTrace/Private/SourceFilterCollection.cpp b/Engine/Plugins/Developer/TraceSourceFiltering/Source/SourceFilteringTrace/Private/SourceFilterCollection.cpp
--- a/Engine/Plugins/Developer/TraceSourceFiltering/Source/SourceFilteringTrace/Private/SourceFilterCollection.cpp
+++ b/Engine/Plugins/Developer/TraceSourceFiltering/Source/SourceFilteringTrace/Private/SourceFilterCollection.cpp
@@ -2,8 +2,8 @@
 
 #include "SourceFilterCollection.h"
 #include "SourceFilterTrace.h"
-#include "AssetRegistryModule.h"
-#include "AssetData.h"
+#include "AssetRegistry/AssetRegistryModule.h"
+#include "AssetRegistry/AssetData.h"
 #include "Async/Async.h"
 #include "UObject/UObjectIterator.h"
 #include "EmptySourceFilter.h"
diff --git a/Engine/Plugins/Editor/AssetManagerEditor/Source/AssetManagerEditor/Private/AssetManagerEditorModule.cpp b/Engine/Plugins/Editor/AssetManagerEditor/Source/AssetManagerEditor/Private/AssetManagerEditorModule.cpp
--- a/Engine/Plugins/Editor/AssetManagerEditor/Source/AssetManagerEditor/Private/AssetManagerEditorModule.cpp
+++ b/Engine/Plugins/Editor/AssetManagerEditor/Source/AssetManagerEditor/Private/AssetManagerEditorModule.cpp
@@ -14,7 +14,7 @@
 #include "CollectionManagerModule.h"
 #include "GameDelegates.h"
 #include "ICollectionManager.h"
-#include "ARFilter.h"
+#include "AssetRegistry/ARFilter.h"
 #include "Misc/FileHelper.h"
 #include "ProfilingDebugging/ProfilingHelpers.h"
 #include "Stats/StatsMisc.h"
@@ -37,7 +37,7 @@
 
 #include "LevelEditor.h"
 #include "GraphEditorModule.h"
-#include "AssetData.h"
+#include "AssetRegistry/AssetData.h"
 #include "Engine/World.h"
 #include "Misc/App.h"
 #include "GenericPlatform/GenericPlatformFile.h"
diff --git a/Engine/Plugins/Editor/DataValidation/Source/DataValidation/Private/DataValidationCommandlet.cpp b/Engine/Plugins/Editor/DataValidation/Source/DataValidation/Private/DataValidationCommandlet.cpp
--- a/Engine/Plugins/Editor/DataValidation/Source/DataValidation/Private/DataValidationCommandlet.cpp
+++ b/Engine/Plugins/Editor/DataValidation/Source/DataValidation/Private/DataValidationCommandlet.cpp
@@ -2,9 +2,9 @@
 
 #include "DataValidationCommandlet.h"
 #include "DataValidationModule.h"
-#include "IAssetRegistry.h"
-#include "AssetRegistryHelpers.h"
-#include "AssetRegistryModule.h"
+#include "AssetRegistry/IAssetRegistry.h"
+#include "AssetRegistry/AssetRegistryHelpers.h"
+#include "AssetRegistry/AssetRegistryModule.h"
 #include "Editor.h"
 #include "EditorValidatorSubsystem.h"
 
diff --git a/Engine/Plugins/FX/Niagara/Source/NiagaraEditor/Private/Commandlets/NiagaraDumpBytecodeCommandlet.cpp b/Engine/Plugins/FX/Niagara/Source/NiagaraEditor/Private/Commandlets/NiagaraDumpBytecodeCommandlet.cpp
--- a/Engine/Plugins/FX/Niagara/Source/NiagaraEditor/Private/Commandlets/NiagaraDumpBytecodeCommandlet.cpp
+++ b/Engine/Plugins/FX/Niagara/Source/NiagaraEditor/Private/Commandlets/NiagaraDumpBytecodeCommandlet.cpp
@@ -2,9 +2,9 @@
 
 #include "Commandlets/NiagaraDumpBytecodeCommandlet.h"
 
-#include "ARFilter.h"
-#include "AssetData.h"
-#include "AssetRegistryModule.h"
+#include "AssetRegistry/ARFilter.h"
+#include "AssetRegistry/AssetData.h"
+#include "AssetRegistry/AssetRegistryModule.h"
 #include "CollectionManagerModule.h"
 #include "CollectionManagerTypes.h"
 #include "HAL/FileManager.h"
diff --git a/Engine/Plugins/FX/Niagara/Source/NiagaraEditor/Private/ViewModels/Stack/NiagaraStackParameterStoreEntry.cpp b/Engine/Plugins/FX/Niagara/Source/NiagaraEditor/Private/ViewModels/Stack/NiagaraStackParameterStoreEntry.cpp
--- a/Engine/Plugins/FX/Niagara/Source/NiagaraEditor/Private/ViewModels/Stack/NiagaraStackParameterStoreEntry.cpp
+++ b/Engine/Plugins/FX/Niagara/Source/NiagaraEditor/Private/ViewModels/Stack/NiagaraStackParameterStoreEntry.cpp
@@ -20,8 +20,8 @@
 #include "Editor.h"
 #include "UObject/StructOnScope.h"
 #include "Modules/ModuleManager.h"
-#include "AssetRegistryModule.h"
-#include "ARFilter.h"
+#include "AssetRegistry/AssetRegistryModule.h"
+#include "AssetRegistry/ARFilter.h"
 #include "EdGraph/EdGraphPin.h"
 #include "NiagaraConstants.h"
 
diff --git a/Engine/Plugins/Media/MediaFrameworkUtilities/Source/MediaFrameworkUtilitiesEditor/Private/MediaFrameworkUtilitiesPlacement.cpp b/Engine/Plugins/Media/MediaFrameworkUtilities/Source/MediaFrameworkUtilitiesEditor/Private/MediaFrameworkUtilitiesPlacement.cpp
--- a/Engine/Plugins/Media/MediaFrameworkUtilities/Source/MediaFrameworkUtilitiesEditor/Private/MediaFrameworkUtilitiesPlacement.cpp
+++ b/Engine/Plugins/Media/MediaFrameworkUtilities/Source/MediaFrameworkUtilitiesEditor/Private/MediaFrameworkUtilitiesPlacement.cpp
@@ -14,11 +14,11 @@
 #include "Widgets/Views/STableRow.h"
 #include "Widgets/Views/SListView.h"
 
-#include "AssetData.h"
-#include "AssetRegistryModule.h"
+#include "AssetRegistry/AssetData.h"
+#include "AssetRegistry/AssetRegistryModule.h"
 #include "DragAndDrop/AssetDragDropOp.h"
 #include "IPlacementModeModule.h"
-#include "IAssetRegistry.h"
+#include "AssetRegistry/IAssetRegistry.h"
 #include "LevelEditor.h"
 
 #include "MediaBundle.h"
diff --git a/Engine/Plugins/MovieScene/CustomizableSequencerTracks/Source/CustomizableSequencerTracksEditor/Private/SequencerTrackBPEditor.cpp b/Engine/Plugins/MovieScene/CustomizableSequencerTracks/Source/CustomizableSequencerTracksEditor/Private/SequencerTrackBPEditor.cpp
--- a/Engine/Plugins/MovieScene/CustomizableSequencerTracks/Source/CustomizableSequencerTracksEditor/Private/SequencerTrackBPEditor.cpp
+++ b/Engine/Plugins/MovieScene/CustomizableSequencerTracks/Source/CustomizableSequencerTracksEditor/Private/SequencerTrackBPEditor.cpp
@@ -7,8 +7,8 @@
 #include "SequencerUtilities.h"
 #include "CustomizableSequencerTracksStyle.h"
 
-#include "AssetRegistryModule.h"
-#include "AssetDataTagMap.h"
+#include "AssetRegistry/AssetRegistryModule.h"
+#include "AssetRegistry/AssetDataTagMap.h"
 #include "UObject/UObjectIterator.h"
 #include "Misc/ScopedSlowTask.h"
 #include "Engine/Blueprint.h"
diff --git a/Engine/Plugins/MovieScene/MovieRenderPipeline/Source/MovieRenderPipelineCore/Private/MoviePipelineUtils.cpp b/Engine/Plugins/MovieScene/MovieRenderPipeline/Source/MovieRenderPipelineCore/Private/MoviePipelineUtils.cpp
--- a/Engine/Plugins/MovieScene/MovieRenderPipeline/Source/MovieRenderPipelineCore/Private/MoviePipelineUtils.cpp
+++ b/Engine/Plugins/MovieScene/MovieRenderPipeline/Source/MovieRenderPipelineCore/Private/MoviePipelineUtils.cpp
@@ -2,10 +2,10 @@
 
 #include "MoviePipelineUtils.h"
 #include "Modules/ModuleManager.h"
-#include "AssetRegistryModule.h"
+#include "AssetRegistry/AssetRegistryModule.h"
 #include "MoviePipelineSetting.h"
-#include "ARFilter.h"
-#include "AssetData.h"
+#include "AssetRegistry/ARFilter.h"
+#include "AssetRegistry/AssetData.h"
 #include "UObject/UObjectIterator.h"
 #include "HAL/IConsoleManager.h"
 #include "MoviePipelineAntiAliasingSetting.h"
diff --git a/Engine/Plugins/Runtime/AssetTags/Source/AssetTags/Private/AssetTagsSubsystem.cpp b/Engine/Plugins/Runtime/AssetTags/Source/AssetTags/Private/AssetTagsSubsystem.cpp
--- a/Engine/Plugins/Runtime/AssetTags/Source/AssetTags/Private/AssetTagsSubsystem.cpp
+++ b/Engine/Plugins/Runtime/AssetTags/Source/AssetTags/Private/AssetTagsSubsystem.cpp
@@ -2,9 +2,9 @@
 
 #include "AssetTagsSubsystem.h"
 
-#include "IAssetRegistry.h"
-#include "AssetRegistryState.h"
-#include "AssetRegistryModule.h"
+#include "AssetRegistry/IAssetRegistry.h"
+#include "AssetRegistry/AssetRegistryState.h"
+#include "AssetRegistry/AssetRegistryModule.h"
 #include "Containers/StringView.h"
 
 #if WITH_EDITOR
diff --git a/Engine/Plugins/VirtualProduction/Takes/Source/TakeRecorder/Private/Widgets/STakeRecorderPanel.cpp b/Engine/Plugins/VirtualProduction/Takes/Source/TakeRecorder/Private/Widgets/STakeRecorderPanel.cpp
--- a/Engine/Plugins/VirtualProduction/Takes/Source/TakeRecorder/Private/Widgets/STakeRecorderPanel.cpp
+++ b/Engine/Plugins/VirtualProduction/Takes/Source/TakeRecorder/Private/Widgets/STakeRecorderPanel.cpp
@@ -23,9 +23,9 @@
 #include "Misc/FileHelper.h"
 
 // AssetRegistry includes
-#include "AssetRegistryModule.h"
-#include "IAssetRegistry.h"
-#include "AssetData.h"
+#include "AssetRegistry/AssetRegistryModule.h"
+#include "AssetRegistry/IAssetRegistry.h"
+#include "AssetRegistry/AssetData.h"
 
 // ContentBrowser includes
 #include "IContentBrowserSingleton.h"
