Re-open media capture when output config changes

diff --git a/Engine/Plugins/Media/MediaFrameworkUtilities/Source/MediaFrameworkUtilitiesEditor/Private/CaptureTab/SMediaFrameworkCapture.cpp b/Engine/Plugins/Media/MediaFrameworkUtilities/Source/MediaFrameworkUtilitiesEditor/Private/CaptureTab/SMediaFrameworkCapture.cpp
--- a/Engine/Plugins/Media/MediaFrameworkUtilities/Source/MediaFrameworkUtilitiesEditor/Private/CaptureTab/SMediaFrameworkCapture.cpp
+++ b/Engine/Plugins/Media/MediaFrameworkUtilities/Source/MediaFrameworkUtilitiesEditor/Private/CaptureTab/SMediaFrameworkCapture.cpp
@@ -177,6 +177,15 @@ SMediaFrameworkCapture::~SMediaFrameworkCapture()
 	FEditorDelegates::PostPIEStarted.RemoveAll(this);
 	FEditorDelegates::PreBeginPIE.RemoveAll(this);
 	EnabledCapture(false);
+	
+	for (const TWeakObjectPtr<UMediaOutput>& WeakMediaOutput : MediaOutputsWithDelegates)
+	{
+		if (UMediaOutput* MediaOutput = WeakMediaOutput.Get())
+		{
+			MediaOutput->OnOutputModified().RemoveAll(this);
+		}
+	}
+	MediaOutputsWithDelegates.Reset();
 }
 
 BEGIN_SLATE_FUNCTION_BUILD_OPTIMIZATION
@@ -382,6 +391,29 @@ TSharedRef<SWidget> SMediaFrameworkCapture::CreateSettingsMenu()
 			NAME_None,
 			EUserInterfaceActionType::ToggleButton
 			);
+
+		SettingsMenuBuilder.AddMenuEntry(
+			LOCTEXT("AutoRestartCaptureOnChange_Label", "Auto Restart On Change"),
+			LOCTEXT("AutoRestartCaptureOnChange_Tooltip", "Should the capture be restarted if the media output is modified."),
+			FSlateIcon(),
+			FUIAction(
+				FExecuteAction::CreateLambda([this]
+					{
+						const FScopedTransaction Transaction(LOCTEXT("AutoRestartCaptureOnChange", "Auto Restart Capture On Change"));
+						GetMutableDefault<UMediaFrameworkEditorCaptureSettings>()->Modify();
+						GetMutableDefault<UMediaFrameworkEditorCaptureSettings>()->bAutoRestartCaptureOnChange = !GetDefault<UMediaFrameworkEditorCaptureSettings>()->bAutoRestartCaptureOnChange;
+						DetailView->SetObject(FindOrAddMediaFrameworkAssetUserData());
+						GetMutableDefault<UMediaFrameworkEditorCaptureSettings>()->SaveConfig();
+					}),
+				FCanExecuteAction(),
+				FIsActionChecked::CreateLambda([this]
+					{
+						return GetDefault<UMediaFrameworkEditorCaptureSettings>()->bAutoRestartCaptureOnChange;
+					})
+				),
+			NAME_None,
+			EUserInterfaceActionType::ToggleButton
+			);
 	}
 
 	return SettingsMenuBuilder.MakeWidget();
@@ -435,6 +467,15 @@ void SMediaFrameworkCapture::EnabledCapture(bool bEnabled)
 		if (bIsCapturing)
 		{
 			EnabledCapture(false);
+			for (const TWeakObjectPtr<UMediaOutput>& WeakMediaOutput : MediaOutputsWithDelegates)
+			{
+				if (UMediaOutput* MediaOutput = WeakMediaOutput.Get())
+				{
+					MediaOutput->OnOutputModified().RemoveAll(this);
+				}
+			}
+			
+			MediaOutputsWithDelegates.Reset();
 		}
 
 		UMediaFrameworkWorldSettingsAssetUserData* AssetUserData = FindOrAddMediaFrameworkAssetUserData();
@@ -459,6 +500,9 @@ void SMediaFrameworkCapture::EnabledCapture(bool bEnabled)
 				.CaptureOptions(Info.CaptureOptions)
 				.ViewMode(Info.ViewMode);
 
+			Info.MediaOutput->OnOutputModified().AddSP(this, &SMediaFrameworkCapture::OnOutputModified);
+			MediaOutputsWithDelegates.Add(Info.MediaOutput);
+
 			CaptureBoxes->AddCaptureWidget(CaptureCameraViewport);
 			CaptureCameraViewport->StartOutput();
 			CaptureCameraViewports.Add(CaptureCameraViewport);
@@ -472,6 +516,9 @@ void SMediaFrameworkCapture::EnabledCapture(bool bEnabled)
 				.CaptureOptions(Info.CaptureOptions)
 				.RenderTarget(Info.RenderTarget);
 
+			Info.MediaOutput->OnOutputModified().AddSP(this, &SMediaFrameworkCapture::OnOutputModified);
+			MediaOutputsWithDelegates.Add(Info.MediaOutput);
+
 			CaptureBoxes->AddCaptureWidget(CaptureRenderTarget);
 			CaptureRenderTarget->StartOutput();
 			CaptureRenderTargets.Add(CaptureRenderTarget);
@@ -485,6 +532,9 @@ void SMediaFrameworkCapture::EnabledCapture(bool bEnabled)
 				.CaptureOptions(AssetUserData->CurrentViewportMediaOutput.CaptureOptions)
 				.ViewMode(AssetUserData->CurrentViewportMediaOutput.ViewMode);
 
+			AssetUserData->CurrentViewportMediaOutput.MediaOutput->OnOutputModified().AddSP(this, &SMediaFrameworkCapture::OnOutputModified);
+			MediaOutputsWithDelegates.Add(AssetUserData->CurrentViewportMediaOutput.MediaOutput);
+
 			CaptureBoxes->AddCaptureWidget(CaptureCurrentViewport);
 			CaptureCurrentViewport->StartOutput();
 		}
@@ -640,6 +690,44 @@ void SMediaFrameworkCapture::OnObjectPreEditChange(UObject* Object, const FEditP
 	}
 }
 
+void SMediaFrameworkCapture::OnOutputModified(UMediaOutput* Output)
+{
+	if (bIsCapturing)
+	{
+		if (GetDefault<UMediaFrameworkEditorCaptureSettings>()->bAutoRestartCaptureOnChange)
+		{
+			/** Only trigger one restart if multiple properties were modified. */
+			if (!NextTickTimerHandle.IsValid())
+			{
+				for (const TSharedPtr<SMediaFrameworkCaptureCameraViewportWidget>& CaptureWidget : CaptureCameraViewports)
+				{
+					CaptureWidget->StopOutput();
+				}
+				
+				for (const TSharedPtr<SMediaFrameworkCaptureRenderTargetWidget>& CaptureWidget : CaptureRenderTargets)
+				{
+					CaptureWidget->StopOutput();
+				}
+				
+				if (CaptureCurrentViewport)
+				{
+					CaptureCurrentViewport->StopOutput();
+				}
+
+				NextTickTimerHandle = GEditor->GetTimerManager()->SetTimerForNextTick(FTimerDelegate::CreateLambda(
+					[]() mutable
+				{
+					if (TSharedPtr<SMediaFrameworkCapture> Capture = WidgetInstance.Pin())
+					{
+						Capture->EnabledCapture(true);
+						Capture->NextTickTimerHandle.Invalidate();
+					}
+				}));
+			}
+		}
+	}
+}
+
 void SMediaFrameworkCapture::OnPrePIE(bool)
 {
 	if (CaptureBoxes.IsValid())
diff --git a/Engine/Plugins/Media/MediaFrameworkUtilities/Source/MediaFrameworkUtilitiesEditor/Private/CaptureTab/SMediaFrameworkCapture.h b/Engine/Plugins/Media/MediaFrameworkUtilities/Source/MediaFrameworkUtilitiesEditor/Private/CaptureTab/SMediaFrameworkCapture.h
--- a/Engine/Plugins/Media/MediaFrameworkUtilities/Source/MediaFrameworkUtilitiesEditor/Private/CaptureTab/SMediaFrameworkCapture.h
+++ b/Engine/Plugins/Media/MediaFrameworkUtilities/Source/MediaFrameworkUtilitiesEditor/Private/CaptureTab/SMediaFrameworkCapture.h
@@ -50,6 +50,10 @@ public:
 	/** Should the Capture setting be saved with the level or should it be saved as a project settings. */
 	UPROPERTY(config)
 	bool bSaveCaptureSetingsInWorld = true;
+	
+	/** Should the capture be restarted if the media output is modified. */
+	UPROPERTY(config)
+	bool bAutoRestartCaptureOnChange = true;
 };
 
 /*
@@ -85,6 +89,7 @@ private:
 	void OnLevelActorsRemoved(AActor* InActor);
 	void OnAssetsDeleted(const TArray<UClass*>& DeletedAssetClasses);
 	void OnObjectPreEditChange(UObject* Object, const FEditPropertyChain& PropertyChain);
+	void OnOutputModified(UMediaOutput* Output);
 
 	TSharedRef<class SWidget> MakeToolBar();
 	TSharedRef<SWidget> CreateSettingsMenu();
@@ -107,4 +112,10 @@ private:
 	TSharedPtr<SMediaFrameworkCaptureCurrentViewportWidget> CaptureCurrentViewport;
 
 	static FDelegateHandle LevelEditorTabManagerChangedHandle;
+
+	/** Handle for the timer used to restart a source after its output options have been modified. */
+	FTimerHandle NextTickTimerHandle;
+
+	/** List of MediaOutputs for which we have a registered delegates. */
+	TArray<TWeakObjectPtr<UMediaOutput>> MediaOutputsWithDelegates;
 };
diff --git a/Engine/Plugins/Media/MediaIOFramework/Source/MediaIOCore/Private/MediaOutput.cpp b/Engine/Plugins/Media/MediaIOFramework/Source/MediaIOCore/Private/MediaOutput.cpp
--- a/Engine/Plugins/Media/MediaIOFramework/Source/MediaIOCore/Private/MediaOutput.cpp
+++ b/Engine/Plugins/Media/MediaIOFramework/Source/MediaIOCore/Private/MediaOutput.cpp
@@ -54,3 +54,12 @@ bool UMediaOutput::Validate(FString& OutFailureReason) const
 
 	return true;
 }
+
+#if WITH_EDITOR
+void UMediaOutput::PostEditChangeProperty(FPropertyChangedEvent& PropertyChangedEvent)
+{
+	UObject::PostEditChangeProperty(PropertyChangedEvent);
+
+	OnOutputModifiedDelegate.Broadcast(this);
+}
+#endif
diff --git a/Engine/Plugins/Media/MediaIOFramework/Source/MediaIOCore/Public/MediaCapture.h b/Engine/Plugins/Media/MediaIOFramework/Source/MediaIOCore/Public/MediaCapture.h
--- a/Engine/Plugins/Media/MediaIOFramework/Source/MediaIOCore/Public/MediaCapture.h
+++ b/Engine/Plugins/Media/MediaIOFramework/Source/MediaIOCore/Public/MediaCapture.h
@@ -146,11 +146,10 @@ public:
 	UPROPERTY(BlueprintReadWrite, EditAnywhere, Category = "MediaCapture")
 	bool bAutostopOnCapture;
 
+
 	/** The number of images to capture*/
-	UPROPERTY(BlueprintReadWrite, EditAnywhere, Category = "MediaCapture")
+	UPROPERTY(BlueprintReadWrite, EditAnywhere, Category = "MediaCapture", meta=(editcondition="bAutostopOnCapture"))
 	int32 NumberOfFramesToCapture;
-
-
 };
 
 
diff --git a/Engine/Plugins/Media/MediaIOFramework/Source/MediaIOCore/Public/MediaOutput.h b/Engine/Plugins/Media/MediaIOFramework/Source/MediaIOCore/Public/MediaOutput.h
--- a/Engine/Plugins/Media/MediaIOFramework/Source/MediaIOCore/Public/MediaOutput.h
+++ b/Engine/Plugins/Media/MediaIOFramework/Source/MediaIOCore/Public/MediaOutput.h
@@ -83,7 +83,19 @@ public:
 	 * The conversion we wish to accomplish on the GPU before the DMA transfer occurs.
 	 */
 	virtual EMediaCaptureConversionOperation GetConversionOperation(EMediaCaptureSourceType InSourceType) const { return EMediaCaptureConversionOperation::NONE; }
+	
+#if WITH_EDITOR
+	DECLARE_MULTICAST_DELEGATE_OneParam(FOnOutputModified, UMediaOutput*);
+	FOnOutputModified& OnOutputModified() { return OnOutputModifiedDelegate; }
+	virtual void PostEditChangeProperty(struct FPropertyChangedEvent& PropertyChangedEvent) override;
+#endif
 
 protected:
 	virtual UMediaCapture* CreateMediaCaptureImpl() { return nullptr; }
+
+private:
+#if WITH_EDITOR
+	FOnOutputModified OnOutputModifiedDelegate;
+#endif
+
 };
