Add an updated "Media Profile" toolbar icon

#rnx

diff --git a/Engine/Plugins/Media/MediaFrameworkUtilities/Source/MediaFrameworkUtilitiesEditor/Private/UI/MediaFrameworkUtilitiesEditorStyle.cpp b/Engine/Plugins/Media/MediaFrameworkUtilities/Source/MediaFrameworkUtilitiesEditor/Private/UI/MediaFrameworkUtilitiesEditorStyle.cpp
--- a/Engine/Plugins/Media/MediaFrameworkUtilities/Source/MediaFrameworkUtilitiesEditor/Private/UI/MediaFrameworkUtilitiesEditorStyle.cpp
+++ b/Engine/Plugins/Media/MediaFrameworkUtilities/Source/MediaFrameworkUtilitiesEditor/Private/UI/MediaFrameworkUtilitiesEditorStyle.cpp
@@ -37,7 +37,7 @@ void FMediaFrameworkUtilitiesEditorStyle::Register()
 	MediaFrameworkUtilitiesStyle::StyleInstance->Set("ClassThumbnail.MediaProfile", new IMAGE_BRUSH("MediaProfile_64x", MediaFrameworkUtilitiesStyle::Icon64x64));
 	MediaFrameworkUtilitiesStyle::StyleInstance->Set("ClassIcon.MediaProfile", new IMAGE_BRUSH("MediaProfile_20x", MediaFrameworkUtilitiesStyle::Icon20x20));
 
-	MediaFrameworkUtilitiesStyle::StyleInstance->Set("ToolbarIcon.MediaProfile", new IMAGE_BRUSH("MediaProfile_Color_40x", MediaFrameworkUtilitiesStyle::Icon40x40));
+	MediaFrameworkUtilitiesStyle::StyleInstance->Set("ToolbarIcon.MediaProfile", new IMAGE_BRUSH_SVG("MediaProfile", MediaFrameworkUtilitiesStyle::Icon40x40));
 	MediaFrameworkUtilitiesStyle::StyleInstance->Set("ToolbarIcon.Apply", new IMAGE_BRUSH("Apply_40x", MediaFrameworkUtilitiesStyle::Icon40x40));
 
 	MediaFrameworkUtilitiesStyle::StyleInstance->Set("TabIcons.MediaCapture.Small", new IMAGE_BRUSH_SVG("MediaCapture", MediaFrameworkUtilitiesStyle::Icon16x16));
