Fix media profile project settings to save to default ini when using the custom "configure now" path

diff --git a/Engine/Plugins/Media/MediaFrameworkUtilities/Source/MediaFrameworkUtilitiesEditor/Private/Profile/MediaProfileSettingsCustomization.cpp b/Engine/Plugins/Media/MediaFrameworkUtilities/Source/MediaFrameworkUtilitiesEditor/Private/Profile/MediaProfileSettingsCustomization.cpp
--- a/Engine/Plugins/Media/MediaFrameworkUtilities/Source/MediaFrameworkUtilitiesEditor/Private/Profile/MediaProfileSettingsCustomization.cpp
+++ b/Engine/Plugins/Media/MediaFrameworkUtilities/Source/MediaFrameworkUtilitiesEditor/Private/Profile/MediaProfileSettingsCustomization.cpp
@@ -287,7 +287,8 @@ void FMediaProfileSettingsCustomization::Configure(const FMediaProfileSettingsCu
 			bool bOnlyIfIsDirty = false;
 			UEditorLoadingAndSavingUtils::SavePackages(PackagesToSave, bOnlyIfIsDirty);
 
-			GetMutableDefault<UMediaProfileSettings>()->SaveConfig();
+			// This will log a warning if the default file was read-only
+			GetMutableDefault<UMediaProfileSettings>()->TryUpdateDefaultConfigFile();
 
 			// Reapply the media profile if it exist
 			UMediaProfile* MediaProfile = IMediaProfileManager::Get().GetCurrentMediaProfile();
