- Updating media texture created by bundle to use new output format

diff --git a/Engine/Plugins/Media/MediaFrameworkUtilities/Source/MediaFrameworkUtilities/Private/MediaBundle.cpp b/Engine/Plugins/Media/MediaFrameworkUtilities/Source/MediaFrameworkUtilities/Private/MediaBundle.cpp
--- a/Engine/Plugins/Media/MediaFrameworkUtilities/Source/MediaFrameworkUtilities/Private/MediaBundle.cpp
+++ b/Engine/Plugins/Media/MediaFrameworkUtilities/Source/MediaFrameworkUtilities/Private/MediaBundle.cpp
@@ -102,7 +102,7 @@ UMediaBundle::UMediaBundle(const FObjectInitializer& ObjectInitializer)
 	: Super(ObjectInitializer)
 {
 #if WITH_EDITOR && WITH_EDITORONLY_DATA
-	static ConstructorHelpers::FObjectFinder<UMaterial> DefaultMaterialFinder(TEXT("/MediaFrameworkUtilities/M_DefaultMedia"));
+	static ConstructorHelpers::FObjectFinder<UMaterial> DefaultMaterialFinder(TEXT("/MediaFrameworkUtilities/M_BaseMediaBundle"));
 	static ConstructorHelpers::FObjectFinder<UTexture> DefaultFailedTextureFinder(TEXT("/MediaFrameworkUtilities/T_VideoInputFailed"));
 	static ConstructorHelpers::FClassFinder<AMediaBundleActorBase> DefaultActorClassFinder(TEXT("/MediaFrameworkUtilities/BP_MediaBundle_Plane_16-9"));
 
@@ -286,6 +286,7 @@ TArray<UPackage*> UMediaBundle::CreateInternalsEditor()
 		//Create MediaTexture 
 		AssetTools.CreateUniqueAssetName(*(ParentName + TEXT("/T_") + GetName() + TEXT("_BC")), TEXT(""), OutPackageName, OutAssetName);
 		MediaTexture = Cast<UMediaTexture>(AssetTools.CreateAsset(OutAssetName, ParentName, UMediaTexture::StaticClass(), nullptr));
+		MediaTexture->NewStyleOutput = true;
 		MediaTexture->SetDefaultMediaPlayer(MediaPlayer);
 		MediaTexture->SetMediaPlayer(MediaPlayer);
 		MediaTexture->UpdateResource();
