Fixed a dangling pointer into a string builder

#rnx

diff --git a/Engine/Plugins/Media/MediaFrameworkUtilities/Source/MediaFrameworkUtilities/Private/Profile/MediaProfile.cpp b/Engine/Plugins/Media/MediaFrameworkUtilities/Source/MediaFrameworkUtilities/Private/Profile/MediaProfile.cpp
--- a/Engine/Plugins/Media/MediaFrameworkUtilities/Source/MediaFrameworkUtilities/Private/Profile/MediaProfile.cpp
+++ b/Engine/Plugins/Media/MediaFrameworkUtilities/Source/MediaFrameworkUtilities/Private/Profile/MediaProfile.cpp
@@ -21,7 +21,7 @@
 namespace MediaProfileAnalytics
 {
 	template <typename ObjectType>
-	auto JoinObjectNames = [](const TArray<ObjectType*>& InObjects)
+	auto JoinObjectNames = [](const TArray<ObjectType*>& InObjects) -> FString
 	{
 		TArray<FString> SourceNames;
 		Algo::TransformIf(
@@ -29,16 +29,16 @@ namespace MediaProfileAnalytics
 			SourceNames,
 			[](ObjectType* Object)
 			{
-				return !!Object;	
+				return !!Object;
 			},
 			[](ObjectType* Object)
 			{
-				return Object->GetName();	
+				return Object->GetName();
 			});
 
 		TStringBuilder<64> StringBuilder;
 		StringBuilder.Join(SourceNames, TEXT(","));
-		return StringBuilder.ToString();
+		return FString(StringBuilder);
 	};
 }
 #endif
