MediaIO - Fix media profile configuration to try to make default settings file writable when configuring new proxies
#rnx

diff --git a/Engine/Plugins/Media/MediaFrameworkUtilities/Source/MediaFrameworkUtilitiesEditor/MediaFrameworkUtilitiesEditor.Build.cs b/Engine/Plugins/Media/MediaFrameworkUtilities/Source/MediaFrameworkUtilitiesEditor/MediaFrameworkUtilitiesEditor.Build.cs
--- a/Engine/Plugins/Media/MediaFrameworkUtilities/Source/MediaFrameworkUtilitiesEditor/MediaFrameworkUtilitiesEditor.Build.cs
+++ b/Engine/Plugins/Media/MediaFrameworkUtilities/Source/MediaFrameworkUtilitiesEditor/MediaFrameworkUtilitiesEditor.Build.cs
@@ -30,6 +30,7 @@ namespace UnrealBuildTool.Rules
                     "MediaUtils",
 					"PlacementMode",
 					"PropertyEditor",
+					"SharedSettingsWidgets",
 					"Slate",
 					"SlateCore",
 					"TimeManagement",
diff --git a/Engine/Plugins/Media/MediaFrameworkUtilities/Source/MediaFrameworkUtilitiesEditor/Private/Profile/MediaProfileSettingsCustomization.cpp b/Engine/Plugins/Media/MediaFrameworkUtilities/Source/MediaFrameworkUtilitiesEditor/Private/Profile/MediaProfileSettingsCustomization.cpp
--- a/Engine/Plugins/Media/MediaFrameworkUtilities/Source/MediaFrameworkUtilitiesEditor/Private/Profile/MediaProfileSettingsCustomization.cpp
+++ b/Engine/Plugins/Media/MediaFrameworkUtilities/Source/MediaFrameworkUtilitiesEditor/Private/Profile/MediaProfileSettingsCustomization.cpp
@@ -19,6 +19,7 @@
 #include "Framework/Application/SlateApplication.h"
 #include "Interfaces/IMainFrameModule.h"
 #include "Misc/ScopedSlowTask.h"
+#include "SSettingsEditorCheckoutNotice.h"
 #include "Widgets/Images/SImage.h"
 #include "Widgets/Input/SButton.h"
 #include "Widgets/Layout/SWidgetSwitcher.h"
@@ -285,7 +286,21 @@ void FMediaProfileSettingsCustomization::Configure(const FMediaProfileSettingsCu
 			UEditorLoadingAndSavingUtils::SavePackages(PackagesToSave, bOnlyIfIsDirty);
 
 			// This will log a warning if the default file was read-only
-			GetMutableDefault<UMediaProfileSettings>()->TryUpdateDefaultConfigFile();
+			const bool bUpdated = GetMutableDefault<UMediaProfileSettings>()->TryUpdateDefaultConfigFile();
+
+			// If the default settings could not be updated, try once to make the default config file writable, then try updating it again
+			if (!bUpdated)
+			{
+				const FString RelativeConfigFilePath = GetMutableDefault<UMediaProfileSettings>()->GetDefaultConfigFilename();
+				const FString TargetFilePath = FPaths::ConvertRelativePathToFull(RelativeConfigFilePath);
+
+				constexpr bool bForceSourceControlUpdate = false;
+				constexpr bool bShowErrorInNotification = true;
+				if (SettingsHelpers::CheckOutOrAddFile(TargetFilePath, bForceSourceControlUpdate, bShowErrorInNotification))
+				{
+					GetMutableDefault<UMediaProfileSettings>()->TryUpdateDefaultConfigFile();
+				}
+			}
 
 			// Reapply the media profile if it exist
 			UMediaProfile* MediaProfile = IMediaProfileManager::Get().GetCurrentMediaProfile();
