Fix media bundle not being able to be placed in level anymore

diff --git a/Engine/Plugins/Media/MediaFrameworkUtilities/Source/MediaFrameworkUtilitiesEditor/Private/MediaFrameworkUtilitiesEditorModule.cpp b/Engine/Plugins/Media/MediaFrameworkUtilities/Source/MediaFrameworkUtilitiesEditor/Private/MediaFrameworkUtilitiesEditorModule.cpp
--- a/Engine/Plugins/Media/MediaFrameworkUtilities/Source/MediaFrameworkUtilitiesEditor/Private/MediaFrameworkUtilitiesEditorModule.cpp
+++ b/Engine/Plugins/Media/MediaFrameworkUtilities/Source/MediaFrameworkUtilitiesEditor/Private/MediaFrameworkUtilitiesEditorModule.cpp
@@ -18,6 +18,7 @@
 #include "Profile/MediaProfile.h"
 #include "Profile/MediaProfileBlueprintLibrary.h"
 #include "Profile/MediaProfileSettings.h"
+#include "Subsystems/PlacementSubsystem.h"
 #include "VideoInputTab/SMediaFrameworkVideoInput.h"
 #include "UI/MediaProfileMenuEntry.h"
 #include "WorkspaceMenuStructure.h"
@@ -46,7 +47,13 @@ public:
 			FMediaProfileCommands::Register();
 			FMediaFrameworkUtilitiesEditorStyle::Register();
 
-			GEditor->ActorFactories.Add(NewObject<UActorFactoryMediaBundle>());
+			ActorFactoryMediaBundle = TStrongObjectPtr<UActorFactoryMediaBundle>(NewObject<UActorFactoryMediaBundle>());
+
+			GEditor->ActorFactories.Add(ActorFactoryMediaBundle.Get());
+			if (UPlacementSubsystem* PlacementSubsystem = GEditor->GetEditorSubsystem<UPlacementSubsystem>())
+			{
+				PlacementSubsystem->RegisterAssetFactory(ActorFactoryMediaBundle.Get());
+			}
 
 			FMediaFrameworkUtilitiesPlacement::RegisterPlacement();
 
@@ -130,6 +137,11 @@ public:
 
 			GEditor->ActorFactories.RemoveAll([](const UActorFactory* ActorFactory) { return ActorFactory->IsA<UActorFactoryMediaBundle>(); });
 
+			if (UPlacementSubsystem* PlacementSubsystem = GEditor->GetEditorSubsystem<UPlacementSubsystem>())
+			{
+				PlacementSubsystem->UnregisterAssetFactory(ActorFactoryMediaBundle.Get());
+			}
+
 			FMediaFrameworkUtilitiesEditorStyle::Unregister();
 			FMediaProfileCommands::Unregister();
 		}
@@ -138,6 +150,7 @@ public:
 private:
 
 	TArray<TSharedRef<IAssetTypeActions>> RegisteredAssetTypeActions;
+	TStrongObjectPtr<UActorFactoryMediaBundle> ActorFactoryMediaBundle;
 };
 
 
