Fix Media Source OCIO not working with media profiles

diff --git a/Engine/Plugins/Media/MediaFrameworkUtilities/Source/MediaFrameworkUtilities/Private/MediaAssets/ProxyMediaSource.cpp b/Engine/Plugins/Media/MediaFrameworkUtilities/Source/MediaFrameworkUtilities/Private/MediaAssets/ProxyMediaSource.cpp
--- a/Engine/Plugins/Media/MediaFrameworkUtilities/Source/MediaFrameworkUtilities/Private/MediaAssets/ProxyMediaSource.cpp
+++ b/Engine/Plugins/Media/MediaFrameworkUtilities/Source/MediaFrameworkUtilities/Private/MediaAssets/ProxyMediaSource.cpp
@@ -200,6 +200,25 @@ FText UProxyMediaSource::GetMediaOption(const FName& Key, const FText& DefaultVa
 	return Super::GetMediaOption(Key, DefaultValue);
 }
 
+TSharedPtr<IMediaOptions::FDataContainer, ESPMode::ThreadSafe> UProxyMediaSource::GetMediaOption(const FName& Key, const TSharedPtr<FDataContainer, ESPMode::ThreadSafe>& DefaultValue) const
+{
+	// Guard against reentrant calls.
+	if (bMediaOptionGuard)
+	{
+		UE_LOG(LogMediaFrameworkUtilities, Warning, TEXT("UMediaSourceProxy::GetMediaOption - Reentrant calls are not supported. Asset: %s"), *GetPathName());
+		return DefaultValue;
+	}
+	TGuardValue<bool> GettingOptionGuard(bMediaOptionGuard, true);
+
+	UMediaSource* CurrentProxy = GetMediaSource();
+	if (CurrentProxy != nullptr)
+	{
+		return CurrentProxy->GetMediaOption(Key, DefaultValue);
+	}
+
+	return Super::GetMediaOption(Key, DefaultValue);
+}
+
 
 bool UProxyMediaSource::HasMediaOption(const FName& Key) const
 {
diff --git a/Engine/Plugins/Media/MediaFrameworkUtilities/Source/MediaFrameworkUtilities/Public/MediaAssets/ProxyMediaSource.h b/Engine/Plugins/Media/MediaFrameworkUtilities/Source/MediaFrameworkUtilities/Public/MediaAssets/ProxyMediaSource.h
--- a/Engine/Plugins/Media/MediaFrameworkUtilities/Source/MediaFrameworkUtilities/Public/MediaAssets/ProxyMediaSource.h
+++ b/Engine/Plugins/Media/MediaFrameworkUtilities/Source/MediaFrameworkUtilities/Public/MediaAssets/ProxyMediaSource.h
@@ -84,6 +84,7 @@ public:
 	virtual int64 GetMediaOption(const FName& Key, int64 DefaultValue) const override;
 	virtual FString GetMediaOption(const FName& Key, const FString& DefaultValue) const override;
 	virtual FText GetMediaOption(const FName& Key, const FText& DefaultValue) const override;
+	virtual TSharedPtr<IMediaOptions::FDataContainer, ESPMode::ThreadSafe> GetMediaOption(const FName& Key, const TSharedPtr<FDataContainer, ESPMode::ThreadSafe>& DefaultValue) const override;
 	virtual bool HasMediaOption(const FName& Key) const override;
 
 private:
