Media Profiles/Rivermax: Moved the Media Profile application event to the point where all modules are initialized. The original issue was the invalid timecode for Rivermax PTP due to Rivermax library being loaded at a later step. Also fixes packaging issues when Rivermax is involved.

diff --git a/Engine/Plugins/Media/MediaFrameworkUtilities/Source/MediaFrameworkUtilities/Private/MediaFrameworkUtilitiesModule.cpp b/Engine/Plugins/Media/MediaFrameworkUtilities/Source/MediaFrameworkUtilities/Private/MediaFrameworkUtilitiesModule.cpp
--- a/Engine/Plugins/Media/MediaFrameworkUtilities/Source/MediaFrameworkUtilities/Private/MediaFrameworkUtilitiesModule.cpp
+++ b/Engine/Plugins/Media/MediaFrameworkUtilities/Source/MediaFrameworkUtilities/Private/MediaFrameworkUtilitiesModule.cpp
@@ -32,7 +32,7 @@ static TAutoConsoleVariable<FString> CVarMediaUtilsStartupProfile(
 class FMediaFrameworkUtilitiesModule : public IMediaFrameworkUtilitiesModule
 {
 	FMediaProfileManager MediaProfileManager;
-	FDelegateHandle PostEngineInitHandle;
+	FDelegateHandle InitHandle;
 
 	virtual void StartupModule() override
 	{
@@ -136,22 +136,15 @@ class FMediaFrameworkUtilitiesModule : public IMediaFrameworkUtilitiesModule
 
 		if (FApp::CanEverRender() || GetDefault<UMediaProfileSettings>()->bApplyInCommandlet)
 		{
-			if (GEngine && GEngine->IsInitialized())
-			{
-				ApplyMediaProfile();
-			}
-			else
-			{
-				PostEngineInitHandle = FCoreDelegates::OnPostEngineInit.AddLambda(ApplyMediaProfile);
-			}
+			InitHandle = FCoreDelegates::OnFEngineLoopInitComplete.AddLambda(ApplyMediaProfile);
 		}
 	}
 
 	void RemoveStartupMediaProfile()
 	{
-		if (PostEngineInitHandle.IsValid())
+		if (InitHandle.IsValid())
 		{
-			FCoreDelegates::OnPostEngineInit.Remove(PostEngineInitHandle);
+			FCoreDelegates::OnFEngineLoopInitComplete.Remove(InitHandle);
 		}
 
 		if (!IsEngineExitRequested())
diff --git a/Engine/Plugins/VirtualProduction/Rivermax/RivermaxCore/Source/RivermaxCore/Private/RivermaxManager.cpp b/Engine/Plugins/VirtualProduction/Rivermax/RivermaxCore/Source/RivermaxCore/Private/RivermaxManager.cpp
--- a/Engine/Plugins/VirtualProduction/Rivermax/RivermaxCore/Source/RivermaxCore/Private/RivermaxManager.cpp
+++ b/Engine/Plugins/VirtualProduction/Rivermax/RivermaxCore/Source/RivermaxCore/Private/RivermaxManager.cpp
@@ -83,11 +83,13 @@ namespace UE::RivermaxCore::Private
 		}
 
 		//Postpone initialization after all modules have been loaded
-		FCoreDelegates::OnAllModuleLoadingPhasesComplete.AddRaw(this, &FRivermaxManager::InitializeLibrary);
+		InitHandle = FCoreDelegates::OnAllModuleLoadingPhasesComplete.AddRaw(this, &FRivermaxManager::InitializeLibrary);
 	}
 
 	FRivermaxManager::~FRivermaxManager()
 	{
+		FCoreDelegates::OnAllModuleLoadingPhasesComplete.Remove(InitHandle);
+
 		if (bIsCleanupRequired)
 		{
 			rmx_status Status = GetApi()->rmx_cleanup();
@@ -491,7 +493,7 @@ namespace UE::RivermaxCore::Private
 				FSlateNotificationManager::Get().AddNotification(Info);
 			}
 #endif
-			UE_LOG(LogRivermax, Error, TEXT("%s"), *ErrorText.ToString());
+			UE_LOG(LogRivermax, Warning, TEXT("%s"), *ErrorText.ToString());
 		}
 		
 		return bIsLibraryInitialized;
diff --git a/Engine/Plugins/VirtualProduction/Rivermax/RivermaxCore/Source/RivermaxCore/Private/RivermaxManager.h b/Engine/Plugins/VirtualProduction/Rivermax/RivermaxCore/Source/RivermaxCore/Private/RivermaxManager.h
--- a/Engine/Plugins/VirtualProduction/Rivermax/RivermaxCore/Source/RivermaxCore/Private/RivermaxManager.h
+++ b/Engine/Plugins/VirtualProduction/Rivermax/RivermaxCore/Source/RivermaxCore/Private/RivermaxManager.h
@@ -90,6 +90,9 @@ namespace UE::RivermaxCore::Private
 
 		/** List of function pointers to the rivermax DLL. */
 		RIVERMAX_API_FUNCTION_LIST FuncList;
+
+		/** Delegate handle. Used to remove the delegate from subscribed module initialization delegate collection. */
+		FDelegateHandle InitHandle;
 	};
 }
 
