Added a check to make sure that we don't get to Media profile application point after engine is initialized.

diff --git a/Engine/Plugins/Media/MediaFrameworkUtilities/Source/MediaFrameworkUtilities/Private/MediaFrameworkUtilitiesModule.cpp b/Engine/Plugins/Media/MediaFrameworkUtilities/Source/MediaFrameworkUtilities/Private/MediaFrameworkUtilitiesModule.cpp
--- a/Engine/Plugins/Media/MediaFrameworkUtilities/Source/MediaFrameworkUtilities/Private/MediaFrameworkUtilitiesModule.cpp
+++ b/Engine/Plugins/Media/MediaFrameworkUtilities/Source/MediaFrameworkUtilities/Private/MediaFrameworkUtilitiesModule.cpp
@@ -136,6 +136,7 @@ class FMediaFrameworkUtilitiesModule : public IMediaFrameworkUtilitiesModule
 
 		if (FApp::CanEverRender() || GetDefault<UMediaProfileSettings>()->bApplyInCommandlet)
 		{
+			check(!(GEngine && GEngine->IsInitialized()));
 			InitHandle = FCoreDelegates::OnFEngineLoopInitComplete.AddLambda(ApplyMediaProfile);
 		}
 	}
