Fixed Media Framework placement item style

#rnx

diff --git a/Engine/Plugins/Media/MediaFrameworkUtilities/Source/MediaFrameworkUtilitiesEditor/Private/MediaFrameworkUtilitiesPlacement.cpp b/Engine/Plugins/Media/MediaFrameworkUtilities/Source/MediaFrameworkUtilitiesEditor/Private/MediaFrameworkUtilitiesPlacement.cpp
--- a/Engine/Plugins/Media/MediaFrameworkUtilities/Source/MediaFrameworkUtilitiesEditor/Private/MediaFrameworkUtilitiesPlacement.cpp
+++ b/Engine/Plugins/Media/MediaFrameworkUtilities/Source/MediaFrameworkUtilitiesEditor/Private/MediaFrameworkUtilitiesPlacement.cpp
@@ -85,7 +85,7 @@ public:
 
 		TSharedRef< STableRow<TSharedPtr<FMediaPlacementListItem>> > TableRowWidget =
 			SNew(STableRow<TSharedPtr<FMediaPlacementListItem>>, OwnerTable)
-			.Style(FAppStyle::Get(), "ContentBrowser.AssetListView.TableRow")
+			.Style(FAppStyle::Get(), "PlacementBrowser.PlaceableItemRow")
 			.OnDragDetected(this, &SMediaPlacementPalette::OnDraggingListViewWidget);
 
 		// Get the MediaSource thumbnail or the MediaBundle is not loaded
