Fix media profile toolbar icon not showing when window is small

diff --git a/Engine/Plugins/Media/MediaFrameworkUtilities/Source/MediaFrameworkUtilitiesEditor/Private/UI/MediaFrameworkUtilitiesEditorStyle.cpp b/Engine/Plugins/Media/MediaFrameworkUtilities/Source/MediaFrameworkUtilitiesEditor/Private/UI/MediaFrameworkUtilitiesEditorStyle.cpp
--- a/Engine/Plugins/Media/MediaFrameworkUtilities/Source/MediaFrameworkUtilitiesEditor/Private/UI/MediaFrameworkUtilitiesEditorStyle.cpp
+++ b/Engine/Plugins/Media/MediaFrameworkUtilities/Source/MediaFrameworkUtilitiesEditor/Private/UI/MediaFrameworkUtilitiesEditorStyle.cpp
@@ -38,7 +38,7 @@ void FMediaFrameworkUtilitiesEditorStyle::Register()
 	MediaFrameworkUtilitiesStyle::StyleInstance->Set("ClassThumbnail.MediaProfile", new IMAGE_BRUSH("MediaProfile_64x", MediaFrameworkUtilitiesStyle::Icon64x64));
 	MediaFrameworkUtilitiesStyle::StyleInstance->Set("ClassIcon.MediaProfile", new IMAGE_BRUSH("MediaProfile_20x", MediaFrameworkUtilitiesStyle::Icon20x20));
 
-	MediaFrameworkUtilitiesStyle::StyleInstance->Set("ToolbarIcon.MediaProfile", new IMAGE_BRUSH_SVG("MediaProfile", MediaFrameworkUtilitiesStyle::Icon40x40));
+	MediaFrameworkUtilitiesStyle::StyleInstance->Set("ToolbarIcon.Edit", new IMAGE_BRUSH_SVG("MediaProfile", MediaFrameworkUtilitiesStyle::Icon40x40));
 	MediaFrameworkUtilitiesStyle::StyleInstance->Set("ToolbarIcon.Apply", new IMAGE_BRUSH("Apply_40x", MediaFrameworkUtilitiesStyle::Icon40x40));
 
 	MediaFrameworkUtilitiesStyle::StyleInstance->Set("TabIcons.MediaCapture.Small", new IMAGE_BRUSH_SVG("MediaCapture", MediaFrameworkUtilitiesStyle::Icon16x16));
diff --git a/Engine/Plugins/Media/MediaFrameworkUtilities/Source/MediaFrameworkUtilitiesEditor/Private/UI/MediaProfileMenuEntry.cpp b/Engine/Plugins/Media/MediaFrameworkUtilities/Source/MediaFrameworkUtilitiesEditor/Private/UI/MediaProfileMenuEntry.cpp
--- a/Engine/Plugins/Media/MediaFrameworkUtilities/Source/MediaFrameworkUtilitiesEditor/Private/UI/MediaProfileMenuEntry.cpp
+++ b/Engine/Plugins/Media/MediaFrameworkUtilities/Source/MediaFrameworkUtilitiesEditor/Private/UI/MediaProfileMenuEntry.cpp
@@ -46,7 +46,7 @@ struct FMediaProfileMenuEntryImpl
 		UToolMenu* Menu = UToolMenus::Get()->ExtendMenu("LevelEditor.LevelEditorToolBar.User");
 		FToolMenuSection& Section = Menu->FindOrAddSection("MediaProfile");
 
-		auto ButtonTooltipLambda = [this]()
+		auto ButtonTooltipLambda = [this]() -> FText
 		{
 			UMediaProfile* MediaProfile = IMediaProfileManager::Get().GetCurrentMediaProfile();
 			if (MediaProfile == nullptr)
@@ -62,7 +62,7 @@ struct FMediaProfileMenuEntryImpl
 			FMediaProfileCommands::Get().Edit,
 			LOCTEXT("MediaProfile_Label", "Media Profile"),
 			MakeAttributeLambda(ButtonTooltipLambda),
-			FSlateIcon(FMediaFrameworkUtilitiesEditorStyle::GetStyleSetName(), TEXT("ToolbarIcon.MediaProfile"))
+			FSlateIcon(FMediaFrameworkUtilitiesEditorStyle::GetStyleSetName(), TEXT("ToolbarIcon.Edit"))
 		);
 		MediaProfileButtonEntry.SetCommandList(Actions);
 
