[MediaFrameworkUtilities] Adding support for MediaOutput classes that don't specify the output resolution.

Some MediaOutput classes like UFileMediaOutput, UNDIMediaOutput, URivermaxMediaOutput, etc have the option of not explicitly specifying the resolution and instad using the "capture source" resolution.
For the Viewport Capture, the MediaCapture panel relies on the media output to provide the size, therefore, in case a media output class doesn't, it should then provide a default one. It has been added in the settings menu.

diff --git a/Engine/Plugins/Media/MediaFrameworkUtilities/Source/MediaFrameworkUtilitiesEditor/Private/CaptureTab/SMediaFrameworkCapture.cpp b/Engine/Plugins/Media/MediaFrameworkUtilities/Source/MediaFrameworkUtilitiesEditor/Private/CaptureTab/SMediaFrameworkCapture.cpp
--- a/Engine/Plugins/Media/MediaFrameworkUtilities/Source/MediaFrameworkUtilitiesEditor/Private/CaptureTab/SMediaFrameworkCapture.cpp
+++ b/Engine/Plugins/Media/MediaFrameworkUtilities/Source/MediaFrameworkUtilitiesEditor/Private/CaptureTab/SMediaFrameworkCapture.cpp
@@ -14,6 +14,7 @@
 #include "SlateOptMacros.h"
 #include "UI/MediaFrameworkUtilitiesEditorStyle.h"
 #include "Widgets/Docking/SDockTab.h"
+#include "Widgets/Input/SVectorInputBox.h"
 #include "Widgets/Layout/SScrollBox.h"
 #include "Widgets/Layout/SSplitter.h"
 
@@ -409,11 +410,67 @@ TSharedRef<SWidget> SMediaFrameworkCapture::CreateSettingsMenu()
 			NAME_None,
 			EUserInterfaceActionType::ToggleButton
 			);
+
+		SettingsMenuBuilder.BeginSection("ViewportCaptureSection", LOCTEXT("ViewportCaptureSectionHeader", "Viewport Capture"));
+		{
+			typedef SNumericVectorInputBox<int32, FIntPoint, 2> SNumericVectorInputBoxInt2;
+
+			const TSharedRef<SWidget> InputBox = SNew(SNumericVectorInputBoxInt2)
+				.Font(FAppStyle::Get().GetFontStyle("PropertyWindow.NormalFont"))
+				.AllowSpin(false)
+				.bColorAxisLabels(false)
+				.X(this, &SMediaFrameworkCapture::GetDefaultViewportSize, 0)
+				.OnXChanged(this, &SMediaFrameworkCapture::DefaultViewportSizeChanged, 0)
+				.OnXCommitted(this, &SMediaFrameworkCapture::DefaultViewportSizeCommitted, 0)
+				.Y(this, &SMediaFrameworkCapture::GetDefaultViewportSize, 1)
+				.OnYChanged(this, &SMediaFrameworkCapture::DefaultViewportSizeChanged, 1)
+				.OnYCommitted(this, &SMediaFrameworkCapture::DefaultViewportSizeCommitted, 1);
+			
+			SettingsMenuBuilder.AddWidget(InputBox,
+				LOCTEXT("DefaultViewportCaptureSize_Label", "Default Size"),
+				false /*bInNoIndent*/,
+				true /*bInSearchable*/,
+				LOCTEXT("DefaultViewportCaptureSize_Tooltip", "Default resolution to use for viewport capture if it is not specified by the media output.")
+				);
+		}
+
+		SettingsMenuBuilder.EndSection();
 	}
 
 	return SettingsMenuBuilder.MakeWidget();
 }
 
+TOptional<int32> SMediaFrameworkCapture::GetDefaultViewportSize(int32 InValueIndex) const
+{
+	if (InValueIndex >= 0 && InValueIndex < 2)
+	{
+		return GetMutableDefault<UMediaFrameworkMediaCaptureSettings>()->DefaultViewportCaptureSize(InValueIndex);
+	}
+	return TOptional<int32>();
+}
+	
+void SMediaFrameworkCapture::DefaultViewportSizeChanged(int32 InValue, int32 InValueIndex)
+{
+	if (InValueIndex >= 0 && InValueIndex < 2)
+	{
+		GetMutableDefault<UMediaFrameworkMediaCaptureSettings>()->DefaultViewportCaptureSize(InValueIndex) = InValue;
+	}
+}
+	
+void SMediaFrameworkCapture::DefaultViewportSizeCommitted(int32 InValue, ETextCommit::Type InCommitInfo, int32 InValueIndex)
+{
+	if (InCommitInfo == ETextCommit::OnEnter || InCommitInfo == ETextCommit::OnUserMovedFocus)
+	{
+		if (InValueIndex >= 0 && InValueIndex < 2)
+		{
+			const FScopedTransaction Transaction(LOCTEXT("DefaultViewportCaptureSize_Transaction", "Default Viewport Capture Size"));
+			GetMutableDefault<UMediaFrameworkMediaCaptureSettings>()->Modify();
+			GetMutableDefault<UMediaFrameworkMediaCaptureSettings>()->DefaultViewportCaptureSize(InValueIndex) = InValue;
+			GetMutableDefault<UMediaFrameworkMediaCaptureSettings>()->SaveConfig();
+		}
+	}
+}
+
 bool SMediaFrameworkCapture::CanEnableViewport() const
 {
 	UMediaFrameworkWorldSettingsAssetUserData* AssetUserData = FindMediaFrameworkAssetUserData();
diff --git a/Engine/Plugins/Media/MediaFrameworkUtilities/Source/MediaFrameworkUtilitiesEditor/Private/CaptureTab/SMediaFrameworkCapture.h b/Engine/Plugins/Media/MediaFrameworkUtilities/Source/MediaFrameworkUtilitiesEditor/Private/CaptureTab/SMediaFrameworkCapture.h
--- a/Engine/Plugins/Media/MediaFrameworkUtilities/Source/MediaFrameworkUtilitiesEditor/Private/CaptureTab/SMediaFrameworkCapture.h
+++ b/Engine/Plugins/Media/MediaFrameworkUtilities/Source/MediaFrameworkUtilitiesEditor/Private/CaptureTab/SMediaFrameworkCapture.h
@@ -34,6 +34,12 @@ class UMediaFrameworkMediaCaptureSettings : public UObject
 public:
 	UPROPERTY(config)
 	bool bIsVerticalSplitterOrientation = true;
+	
+	/**
+	 * This is the resolution to use if it is not specified by the media output.
+	 */
+	UPROPERTY(EditAnywhere, config, Category="Media Viewport Capture")
+	FIntPoint DefaultViewportCaptureSize = FIntPoint(1920, 1080);
 };
 
 /**
@@ -91,6 +97,10 @@ private:
 	TSharedRef<class SWidget> MakeToolBar();
 	TSharedRef<SWidget> CreateSettingsMenu();
 
+	TOptional<int32> GetDefaultViewportSize(int32 InValueIndex) const;
+	void DefaultViewportSizeChanged(int32 InValue, int32 InValueIndex);
+	void DefaultViewportSizeCommitted(int32 InValue, ETextCommit::Type InCommitInfo, int32 InValueIndex);
+
 	bool CanEnableViewport() const;
 	UMediaFrameworkWorldSettingsAssetUserData* FindMediaFrameworkAssetUserData() const;
 
diff --git a/Engine/Plugins/Media/MediaFrameworkUtilities/Source/MediaFrameworkUtilitiesEditor/Private/CaptureTab/SMediaFrameworkCaptureOutputWidget.cpp b/Engine/Plugins/Media/MediaFrameworkUtilities/Source/MediaFrameworkUtilitiesEditor/Private/CaptureTab/SMediaFrameworkCaptureOutputWidget.cpp
--- a/Engine/Plugins/Media/MediaFrameworkUtilities/Source/MediaFrameworkUtilitiesEditor/Private/CaptureTab/SMediaFrameworkCaptureOutputWidget.cpp
+++ b/Engine/Plugins/Media/MediaFrameworkUtilities/Source/MediaFrameworkUtilitiesEditor/Private/CaptureTab/SMediaFrameworkCaptureOutputWidget.cpp
@@ -7,6 +7,7 @@
 #include "Engine/TextureRenderTarget2D.h"
 #include "Framework/Application/SlateApplication.h"
 #include "LevelEditorViewport.h"
+#include "SMediaFrameworkCapture.h"
 #include "Slate/SceneViewport.h"
 #include "SlateOptMacros.h"
 #include "Widgets/Input/SButton.h"
@@ -28,15 +29,32 @@
  */
 namespace MediaFrameworkUtilities
 {
-	const float Padding = 4.f;
-	const float ViewportBoxDesiredSizeY = 200.f;
-	float GetWidthOverride(TWeakObjectPtr<UMediaOutput> MediaOutput)
+	constexpr float Padding = 4.f;
+	constexpr float ViewportBoxDesiredSizeY = 200.f;
+	
+	/**
+	 * Determines the target size from the media output if desired size is specified, otherwise,
+	 * fallback to the default viewport capture resolution.
+	 */
+	FIntPoint GetTargetSize(const TWeakObjectPtr<UMediaOutput>& InMediaOutputWeak)
+	{
+		if (const UMediaOutput* MediaOutputPtr = InMediaOutputWeak.Get())
+		{
+			const FIntPoint TargetSize = MediaOutputPtr->GetRequestedSize();
+			if (TargetSize != UMediaOutput::RequestCaptureSourceSize)
+			{
+				return TargetSize;
+			}
+		}
+		return GetMutableDefault<UMediaFrameworkMediaCaptureSettings>()->DefaultViewportCaptureSize;
+	}
+	
+	float GetWidthOverride(const TWeakObjectPtr<UMediaOutput> InMediaOutputWeak)
 	{
 		float Ratio = 1.77777777f;
-		UMediaOutput* MediaOutputPtr = MediaOutput.Get();
-		if (MediaOutputPtr)
+		const FIntPoint TargetSize = GetTargetSize(InMediaOutputWeak);
+		if (TargetSize.X > 0 && TargetSize.Y > 0)
 		{
-			FIntPoint TargetSize = MediaOutputPtr->GetRequestedSize();
 			Ratio = (float)TargetSize.X / (float)TargetSize.Y;
 		}
 		return ViewportBoxDesiredSizeY * Ratio;
@@ -371,7 +389,7 @@ void SMediaFrameworkCaptureCameraViewportWidget::StartOutput()
 		MediaCapture.Reset(MediaOutputPtr->CreateMediaCapture());
 		if (MediaCapture.IsValid())
 		{
-			FIntPoint TargetSize = MediaOutputPtr->GetRequestedSize();
+			const FIntPoint TargetSize = MediaFrameworkUtilities::GetTargetSize(MediaOutput);
 			SceneViewport->SetFixedViewportSize(TargetSize.X, TargetSize.Y);
 			MediaCapture->CaptureSceneViewport(SceneViewport, CaptureOptions);
 		}
