Media Profile: Implemented basic framework and media tree for new media profile editor. New editor is feature flagged with cvar "MediaProfile.UseNewMediaProfileEditor" for the moment until more of the UI is implemented

#virtualized

diff --git a/Engine/Plugins/Media/AjaMedia/Source/AjaMedia/Private/Assets/AjaMediaSource.cpp b/Engine/Plugins/Media/AjaMedia/Source/AjaMedia/Private/Assets/AjaMediaSource.cpp
--- a/Engine/Plugins/Media/AjaMedia/Source/AjaMedia/Private/Assets/AjaMediaSource.cpp
+++ b/Engine/Plugins/Media/AjaMedia/Source/AjaMedia/Private/Assets/AjaMediaSource.cpp
@@ -5,6 +5,7 @@
 #include "Aja/Aja.h"
 #include "AjaDeviceProvider.h"
 #include "AjaMediaPrivate.h"
+#include "IMediaIOCoreModule.h"
 #include "MediaIOCorePlayerBase.h"
 
 UAjaMediaSource::UAjaMediaSource()
@@ -281,6 +282,19 @@ bool UAjaMediaSource::Validate() const
 	return true;
 }
 
+#if WITH_EDITOR
+FString UAjaMediaSource::GetDescriptionString() const
+{
+	IMediaIOCoreDeviceProvider* DeviceProviderPtr = IMediaIOCoreModule::Get().GetDeviceProvider(FAjaDeviceProvider::GetProviderName());
+	if (DeviceProviderPtr)
+	{
+		return DeviceProviderPtr->ToText(MediaConfiguration, bAutoDetectInput).ToString(); 
+	}
+
+	return Super::GetDescriptionString();
+}
+#endif //WITH_EDITOR
+
 #if WITH_EDITOR
 bool UAjaMediaSource::CanEditChange(const FProperty* InProperty) const
 {
diff --git a/Engine/Plugins/Media/AjaMedia/Source/AjaMedia/Public/AjaMediaSource.h b/Engine/Plugins/Media/AjaMedia/Source/AjaMedia/Public/AjaMediaSource.h
--- a/Engine/Plugins/Media/AjaMedia/Source/AjaMedia/Public/AjaMediaSource.h
+++ b/Engine/Plugins/Media/AjaMedia/Source/AjaMedia/Public/AjaMediaSource.h
@@ -137,6 +137,10 @@ public:
 
 	virtual FString GetUrl() const override;
 	virtual bool Validate() const override;
+	
+#if WITH_EDITOR
+	virtual FString GetDescriptionString() const override;
+#endif //WITH_EDITOR
 
 
 public:
diff --git a/Engine/Plugins/Media/AjaMedia/Source/AjaMediaOutput/Private/AjaMediaOutput.cpp b/Engine/Plugins/Media/AjaMedia/Source/AjaMediaOutput/Private/AjaMediaOutput.cpp
--- a/Engine/Plugins/Media/AjaMedia/Source/AjaMediaOutput/Private/AjaMediaOutput.cpp
+++ b/Engine/Plugins/Media/AjaMedia/Source/AjaMediaOutput/Private/AjaMediaOutput.cpp
@@ -2,10 +2,12 @@
 
 #include "AjaMediaOutput.h"
 
+#include "AjaDeviceProvider.h"
 #include "AJALib.h"
 #include "AjaMediaCapture.h"
 #include "AjaMediaSettings.h"
 #include "IAjaMediaModule.h"
+#include "IMediaIOCoreModule.h"
 #include "Modules/ModuleManager.h"
 #include "UObject/EnterpriseObjectVersion.h"
 
@@ -203,6 +205,19 @@ EMediaCaptureConversionOperation UAjaMediaOutput::GetConversionOperation(EMediaC
 	return Result;
 }
 
+#if WITH_EDITOR
+FString UAjaMediaOutput::GetDescriptionString() const
+{
+	IMediaIOCoreDeviceProvider* DeviceProviderPtr = IMediaIOCoreModule::Get().GetDeviceProvider(FAjaDeviceProvider::GetProviderName());
+	if (DeviceProviderPtr)
+	{
+		return DeviceProviderPtr->ToText(OutputConfiguration).ToString(); 
+	}
+
+	return Super::GetDescriptionString();
+}
+#endif // WITH_EDITOR
+
 UMediaCapture* UAjaMediaOutput::CreateMediaCaptureImpl()
 {
 	UMediaCapture* Result = NewObject<UAjaMediaCapture>();
diff --git a/Engine/Plugins/Media/AjaMedia/Source/AjaMediaOutput/Public/AjaMediaOutput.h b/Engine/Plugins/Media/AjaMedia/Source/AjaMediaOutput/Public/AjaMediaOutput.h
--- a/Engine/Plugins/Media/AjaMedia/Source/AjaMediaOutput/Public/AjaMediaOutput.h
+++ b/Engine/Plugins/Media/AjaMedia/Source/AjaMediaOutput/Public/AjaMediaOutput.h
@@ -145,6 +145,10 @@ public:
 	virtual EPixelFormat GetRequestedPixelFormat() const override;
 	virtual EMediaCaptureConversionOperation GetConversionOperation(EMediaCaptureSourceType InSourceType) const override;
 
+#if WITH_EDITOR
+	virtual FString GetDescriptionString() const override;
+#endif //WITH_EDITOR
+	
 protected:
 	virtual UMediaCapture* CreateMediaCaptureImpl() override;
 
diff --git a/Engine/Plugins/Media/BlackmagicMedia/Source/BlackmagicMedia/Private/Assets/BlackmagicMediaSource.cpp b/Engine/Plugins/Media/BlackmagicMedia/Source/BlackmagicMedia/Private/Assets/BlackmagicMediaSource.cpp
--- a/Engine/Plugins/Media/BlackmagicMedia/Source/BlackmagicMedia/Private/Assets/BlackmagicMediaSource.cpp
+++ b/Engine/Plugins/Media/BlackmagicMedia/Source/BlackmagicMedia/Private/Assets/BlackmagicMediaSource.cpp
@@ -5,6 +5,7 @@
 #include "BlackmagicDeviceProvider.h"
 #include "BlackmagicMediaPrivate.h"
 #include "IBlackmagicMediaModule.h"
+#include "IMediaIOCoreModule.h"
 
 #include "MediaIOCorePlayerBase.h"
 
@@ -173,6 +174,19 @@ bool UBlackmagicMediaSource::Validate() const
 	return true;
 }
 
+#if WITH_EDITOR
+FString UBlackmagicMediaSource::GetDescriptionString() const
+{
+	IMediaIOCoreDeviceProvider* DeviceProviderPtr = IMediaIOCoreModule::Get().GetDeviceProvider(FBlackmagicDeviceProvider::GetProviderName());
+	if (DeviceProviderPtr)
+	{
+		return DeviceProviderPtr->ToText(MediaConfiguration, bAutoDetectInput).ToString(); 
+	}
+
+	return Super::GetDescriptionString();
+}
+#endif //WITH_EDITOR
+
 void UBlackmagicMediaSource::AssignDefaultConfiguration()
 {
 	const FBlackmagicDeviceProvider DeviceProvider;
diff --git a/Engine/Plugins/Media/BlackmagicMedia/Source/BlackmagicMedia/Public/BlackmagicMediaSource.h b/Engine/Plugins/Media/BlackmagicMedia/Source/BlackmagicMedia/Public/BlackmagicMediaSource.h
--- a/Engine/Plugins/Media/BlackmagicMedia/Source/BlackmagicMedia/Public/BlackmagicMediaSource.h
+++ b/Engine/Plugins/Media/BlackmagicMedia/Source/BlackmagicMedia/Public/BlackmagicMediaSource.h
@@ -124,6 +124,10 @@ public:
 	virtual FString GetUrl() const override;
 	virtual bool Validate() const override;
 
+#if WITH_EDITOR
+	virtual FString GetDescriptionString() const override;
+#endif //WITH_EDITOR
+	
 public:
 	//~ UObject interface
 #if WITH_EDITOR
diff --git a/Engine/Plugins/Media/BlackmagicMedia/Source/BlackmagicMediaOutput/Private/BlackmagicMediaOutput.cpp b/Engine/Plugins/Media/BlackmagicMedia/Source/BlackmagicMediaOutput/Private/BlackmagicMediaOutput.cpp
--- a/Engine/Plugins/Media/BlackmagicMedia/Source/BlackmagicMediaOutput/Private/BlackmagicMediaOutput.cpp
+++ b/Engine/Plugins/Media/BlackmagicMedia/Source/BlackmagicMediaOutput/Private/BlackmagicMediaOutput.cpp
@@ -2,9 +2,12 @@
 
 #include "BlackmagicMediaOutput.h"
 
+#include "BlackmagicDeviceProvider.h"
 #include "BlackmagicLib.h"
 #include "BlackmagicMediaCapture.h"
 #include "IBlackmagicMediaModule.h"
+#include "IMediaIOCoreDeviceProvider.h"
+#include "IMediaIOCoreModule.h"
 #include "Modules/ModuleManager.h"
 
 
@@ -135,6 +138,19 @@ EMediaCaptureConversionOperation UBlackmagicMediaOutput::GetConversionOperation(
 	return Result;
 }
 
+#if WITH_EDITOR
+FString UBlackmagicMediaOutput::GetDescriptionString() const
+{
+	IMediaIOCoreDeviceProvider* DeviceProviderPtr = IMediaIOCoreModule::Get().GetDeviceProvider(FBlackmagicDeviceProvider::GetProviderName());
+	if (DeviceProviderPtr)
+	{
+		return DeviceProviderPtr->ToText(OutputConfiguration).ToString(); 
+	}
+
+	return Super::GetDescriptionString();
+}
+#endif //WITH_EDITOR
+
 UMediaCapture* UBlackmagicMediaOutput::CreateMediaCaptureImpl()
 {
 	UMediaCapture* Result = NewObject<UBlackmagicMediaCapture>();
diff --git a/Engine/Plugins/Media/BlackmagicMedia/Source/BlackmagicMediaOutput/Public/BlackmagicMediaOutput.h b/Engine/Plugins/Media/BlackmagicMedia/Source/BlackmagicMediaOutput/Public/BlackmagicMediaOutput.h
--- a/Engine/Plugins/Media/BlackmagicMedia/Source/BlackmagicMediaOutput/Public/BlackmagicMediaOutput.h
+++ b/Engine/Plugins/Media/BlackmagicMedia/Source/BlackmagicMediaOutput/Public/BlackmagicMediaOutput.h
@@ -144,6 +144,10 @@ public:
 	virtual EPixelFormat GetRequestedPixelFormat() const override;
 	virtual EMediaCaptureConversionOperation GetConversionOperation(EMediaCaptureSourceType InSourceType) const override;
 
+#if WITH_EDITOR
+	virtual FString GetDescriptionString() const override;
+#endif //WITH_EDITOR
+	
 protected:
 	virtual UMediaCapture* CreateMediaCaptureImpl() override;
 
diff --git a/Engine/Plugins/Media/MediaFrameworkUtilities/Source/MediaFrameworkUtilities/Private/MediaAssets/ProxyMediaOutput.cpp b/Engine/Plugins/Media/MediaFrameworkUtilities/Source/MediaFrameworkUtilities/Private/MediaAssets/ProxyMediaOutput.cpp
--- a/Engine/Plugins/Media/MediaFrameworkUtilities/Source/MediaFrameworkUtilities/Private/MediaAssets/ProxyMediaOutput.cpp
+++ b/Engine/Plugins/Media/MediaFrameworkUtilities/Source/MediaFrameworkUtilities/Private/MediaAssets/ProxyMediaOutput.cpp
@@ -87,6 +87,21 @@ EMediaCaptureConversionOperation UProxyMediaOutput::GetConversionOperation(EMedi
 	return (CurrentProxy != nullptr) ? CurrentProxy->GetConversionOperation(InSourceType) : EMediaCaptureConversionOperation::NONE;
 }
 
+#if WITH_EDITOR
+FString UProxyMediaOutput::GetDescriptionString() const
+{
+	// Guard against reentrant calls.
+	if (bConfigurationTextGuard)
+	{
+		UE_LOG(LogMediaFrameworkUtilities, Warning, TEXT("UProxyMediaOutput::GetConfigurationDisplayText - Reentrant calls are not supported. Asset: %s"), *GetPathName());
+		return FString();
+	}
+	TGuardValue<bool> Guard(bConfigurationTextGuard, true);
+
+	UMediaOutput* CurrentProxy = GetMediaOutput();
+	return (CurrentProxy != nullptr) ? CurrentProxy->GetDescriptionString() : FString();
+}
+#endif //WITH_EDITOR
 
 UMediaCapture* UProxyMediaOutput::CreateMediaCaptureImpl()
 {
diff --git a/Engine/Plugins/Media/MediaFrameworkUtilities/Source/MediaFrameworkUtilities/Private/MediaAssets/ProxyMediaSource.cpp b/Engine/Plugins/Media/MediaFrameworkUtilities/Source/MediaFrameworkUtilities/Private/MediaAssets/ProxyMediaSource.cpp
--- a/Engine/Plugins/Media/MediaFrameworkUtilities/Source/MediaFrameworkUtilities/Private/MediaAssets/ProxyMediaSource.cpp
+++ b/Engine/Plugins/Media/MediaFrameworkUtilities/Source/MediaFrameworkUtilities/Private/MediaAssets/ProxyMediaSource.cpp
@@ -49,6 +49,21 @@ bool UProxyMediaSource::Validate() const
 	return (CurrentProxy != nullptr) ? CurrentProxy->Validate() : false;
 }
 
+#if WITH_EDITOR
+FString UProxyMediaSource::GetDescriptionString() const
+{
+	// Guard against reentrant calls.
+	if (bConfigurationTextGuard)
+	{
+		UE_LOG(LogMediaFrameworkUtilities, Warning, TEXT("UMediaSourceProxy::GetConfigurationDisplayText - Reentrant calls are not supported. Asset: %s"), *GetPathName());
+		return FString();
+	}
+	TGuardValue<bool> Guard(bConfigurationTextGuard, true);
+
+	UMediaSource* CurrentProxy = GetMediaSource();
+	return (CurrentProxy != nullptr) ? CurrentProxy->GetDescriptionString() : FString();
+}
+#endif //WITH_EDITOR
 
 /* UMediaSourceProxy implementation
  *****************************************************************************/
diff --git a/Engine/Plugins/Media/MediaFrameworkUtilities/Source/MediaFrameworkUtilities/Private/Profile/MediaProfile.cpp b/Engine/Plugins/Media/MediaFrameworkUtilities/Source/MediaFrameworkUtilities/Private/Profile/MediaProfile.cpp
--- a/Engine/Plugins/Media/MediaFrameworkUtilities/Source/MediaFrameworkUtilities/Private/Profile/MediaProfile.cpp
+++ b/Engine/Plugins/Media/MediaFrameworkUtilities/Source/MediaFrameworkUtilities/Private/Profile/MediaProfile.cpp
@@ -58,6 +58,35 @@ int32 UMediaProfile::NumMediaSources() const
 	return MediaSources.Num();
 }
 
+void UMediaProfile::AddMediaSource(UMediaSource* MediaSource)
+{
+	MediaSources.Add(MediaSource);
+}
+
+bool UMediaProfile::RemoveMediaSource(int32 Index)
+{
+	if (!MediaSources.IsValidIndex(Index))
+	{
+		return false;
+	}
+	
+	MediaSources.RemoveAt(Index, EAllowShrinking::No);
+	return true;
+}
+
+bool UMediaProfile::MoveMediaSource(int32 CurrentIndex, int32 DestIndex)
+{
+	if (!MediaSources.IsValidIndex(CurrentIndex) || !MediaSources.IsValidIndex(DestIndex))
+	{
+		return false;
+	}
+
+	TObjectPtr<UMediaSource> MediaSource = MediaSources[CurrentIndex];
+	MediaSources.RemoveAt(CurrentIndex, EAllowShrinking::No);
+	MediaSources.Insert(MediaSource, DestIndex);
+	
+	return true;
+}
 
 UMediaOutput* UMediaProfile::GetMediaOutput(int32 Index) const
 {
@@ -74,6 +103,35 @@ int32 UMediaProfile::NumMediaOutputs() const
 	return MediaOutputs.Num();
 }
 
+void UMediaProfile::AddMediaOutput(UMediaOutput* MediaOutput)
+{
+	MediaOutputs.Add(MediaOutput);
+}
+
+bool UMediaProfile::RemoveMediaOutput(int32 Index)
+{
+	if (!MediaOutputs.IsValidIndex(Index))
+	{
+		return false;
+	}
+
+	MediaOutputs.RemoveAt(Index, EAllowShrinking::No);
+	return true;
+}
+
+bool UMediaProfile::MoveMediaOutput(int32 CurrentIndex, int32 DestIndex)
+{
+	if (!MediaOutputs.IsValidIndex(CurrentIndex) || !MediaOutputs.IsValidIndex(DestIndex))
+	{
+		return false;
+	}
+
+	TObjectPtr<UMediaOutput> MediaOutput = MediaOutputs[CurrentIndex];
+	MediaOutputs.RemoveAt(CurrentIndex, EAllowShrinking::No);
+	MediaOutputs.Insert(MediaOutput, DestIndex);
+	
+	return true;
+}
 
 UTimecodeProvider* UMediaProfile::GetTimecodeProvider() const
 {
@@ -291,16 +349,33 @@ void UMediaProfile::SendAnalytics() const
 void UMediaProfile::FixNumSourcesAndOutputs()
 {
 	const int32 NumSourceProxies = IMediaProfileManager::Get().GetAllMediaSourceProxy().Num();
-	if (MediaSources.Num() != NumSourceProxies)
+	const int32 NumOutputProxies = IMediaProfileManager::Get().GetAllMediaOutputProxy().Num();
+
+	const bool bResizeMediaSources = MediaSources.Num() != NumSourceProxies;
+	const bool bResizeMediaOutputs = MediaOutputs.Num() != NumOutputProxies;
+
+	if (bResizeMediaSources || bResizeMediaOutputs)
 	{
-		MediaSources.SetNumZeroed(NumSourceProxies);
 		Modify();
 	}
+	
+	if (bResizeMediaSources)
+	{
+		MediaSources.SetNumZeroed(NumSourceProxies);
 
-	const int32 NumOutputProxies = IMediaProfileManager::Get().GetAllMediaOutputProxy().Num();
-	if (MediaOutputs.Num() != NumOutputProxies)
+#if WITH_EDITOR
+		FPropertyChangedEvent PropertyChangedEvent(StaticClass()->FindPropertyByName(GET_MEMBER_NAME_CHECKED(UMediaProfile, MediaSources)), EPropertyChangeType::ArrayAdd);
+		FCoreUObjectDelegates::OnObjectPropertyChanged.Broadcast(this, PropertyChangedEvent);
+#endif // WITH_EDITOR
+	}
+
+	if (bResizeMediaOutputs)
 	{
-		Modify();
 		MediaOutputs.SetNumZeroed(NumOutputProxies);
+
+#if WITH_EDITOR
+		FPropertyChangedEvent PropertyChangedEvent(StaticClass()->FindPropertyByName(GET_MEMBER_NAME_CHECKED(UMediaProfile, MediaOutputs)), EPropertyChangeType::ArrayAdd);
+		FCoreUObjectDelegates::OnObjectPropertyChanged.Broadcast(this, PropertyChangedEvent);
+#endif // WITH_EDITOR
 	}
 }
diff --git a/Engine/Plugins/Media/MediaFrameworkUtilities/Source/MediaFrameworkUtilities/Private/Profile/MediaProfileSettings.cpp b/Engine/Plugins/Media/MediaFrameworkUtilities/Source/MediaFrameworkUtilities/Private/Profile/MediaProfileSettings.cpp
--- a/Engine/Plugins/Media/MediaFrameworkUtilities/Source/MediaFrameworkUtilities/Private/Profile/MediaProfileSettings.cpp
+++ b/Engine/Plugins/Media/MediaFrameworkUtilities/Source/MediaFrameworkUtilities/Private/Profile/MediaProfileSettings.cpp
@@ -7,6 +7,12 @@
 #include "MediaAssets/ProxyMediaOutput.h"
 #include "Profile/MediaProfile.h"
 
+namespace MediaProfileSettings
+{
+	const FString DefaultProxiesPath = TEXT("/MediaFrameworkUtilities/DefaultProxies");
+	const FString DefaultMediaProxySourceName = TEXT("DefaultMediaProxySource");
+	const FString DefaultMediaProxyOutputName = TEXT("DefaultMediaProxyOutput");
+}
 
 TArray<UProxyMediaSource*> UMediaProfileSettings::LoadMediaSourceProxies() const
 {
@@ -51,6 +57,24 @@ void UMediaProfileSettings::SetMediaSourceProxy(const TArray<UProxyMediaSource*>
 	OnMediaProxiesChanged.Broadcast();
 }
 
+void UMediaProfileSettings::FillDefaultMediaSourceProxies(int32 NumSources, bool bSilent)
+{
+	MediaSourceProxy.Reset(NumSources);
+	for (int32 Index = 1; Index <= NumSources; ++Index)
+	{
+		FString IndexStr = Index > 9 ? FString::FromInt(Index) : TEXT("0") + FString::FromInt(Index);
+		FSoftObjectPath ProxyPath(FString::Format(TEXT("{0}/{1}_{2}"), { MediaProfileSettings::DefaultProxiesPath, MediaProfileSettings::DefaultMediaProxySourceName, IndexStr }));
+		UProxyMediaSource* DefaultProxySource = Cast<UProxyMediaSource>(ProxyPath.TryLoad());
+		check(DefaultProxySource);
+
+		MediaSourceProxy.Add(DefaultProxySource);
+	}
+
+	if (!bSilent)
+	{
+		OnMediaProxiesChanged.Broadcast();
+	}
+}
 
 void UMediaProfileSettings::SetMediaOutputProxy(const TArray<UProxyMediaOutput*>& InProxies)
 {
@@ -62,6 +86,24 @@ void UMediaProfileSettings::SetMediaOutputProxy(const TArray<UProxyMediaOutput*>
 	OnMediaProxiesChanged.Broadcast();
 }
 
+void UMediaProfileSettings::FillDefaultMediaOutputProxies(int32 NumOutputs, bool bSilent)
+{
+	MediaOutputProxy.Reset(NumOutputs);
+	for (int32 Index = 1; Index <= NumOutputs; ++Index)
+	{
+		FString IndexStr = Index > 9 ? FString::FromInt(Index) : TEXT("0") + FString::FromInt(Index);
+		FSoftObjectPath ProxyPath(FString::Format(TEXT("{0}/{1}_{2}"), { MediaProfileSettings::DefaultProxiesPath, MediaProfileSettings::DefaultMediaProxyOutputName, IndexStr }));
+		UProxyMediaOutput* DefaultProxyOutput = Cast<UProxyMediaOutput>(ProxyPath.TryLoad());
+		check(DefaultProxyOutput);
+
+		MediaOutputProxy.Add(DefaultProxyOutput);
+	}
+
+	if (!bSilent)
+	{
+		OnMediaProxiesChanged.Broadcast();
+	}
+}
 
 void UMediaProfileSettings::PostEditChangeProperty(struct FPropertyChangedEvent& InPropertyChangedEvent)
 {
diff --git a/Engine/Plugins/Media/MediaFrameworkUtilities/Source/MediaFrameworkUtilities/Public/MediaAssets/ProxyMediaOutput.h b/Engine/Plugins/Media/MediaFrameworkUtilities/Source/MediaFrameworkUtilities/Public/MediaAssets/ProxyMediaOutput.h
--- a/Engine/Plugins/Media/MediaFrameworkUtilities/Source/MediaFrameworkUtilities/Public/MediaAssets/ProxyMediaOutput.h
+++ b/Engine/Plugins/Media/MediaFrameworkUtilities/Source/MediaFrameworkUtilities/Public/MediaAssets/ProxyMediaOutput.h
@@ -77,6 +77,10 @@ public:
 	virtual EPixelFormat GetRequestedPixelFormat() const override;
 	virtual EMediaCaptureConversionOperation GetConversionOperation(EMediaCaptureSourceType InSourceType) const override;
 
+#if WITH_EDITOR
+	virtual FString GetDescriptionString() const override;
+#endif //WITH_EDITOR
+	
 protected:
 
 	virtual UMediaCapture* CreateMediaCaptureImpl() override;
@@ -87,5 +91,6 @@ private:
 	mutable bool bValidateGuard;
 	mutable bool bRequestedSizeGuard;
 	mutable bool bRequestedPixelFormatGuard;
+	mutable bool bConfigurationTextGuard;
 	mutable bool bCreateMediaCaptureImplGuard;
 };
diff --git a/Engine/Plugins/Media/MediaFrameworkUtilities/Source/MediaFrameworkUtilities/Public/MediaAssets/ProxyMediaSource.h b/Engine/Plugins/Media/MediaFrameworkUtilities/Source/MediaFrameworkUtilities/Public/MediaAssets/ProxyMediaSource.h
--- a/Engine/Plugins/Media/MediaFrameworkUtilities/Source/MediaFrameworkUtilities/Public/MediaAssets/ProxyMediaSource.h
+++ b/Engine/Plugins/Media/MediaFrameworkUtilities/Source/MediaFrameworkUtilities/Public/MediaAssets/ProxyMediaSource.h
@@ -75,6 +75,10 @@ public:
 	virtual FString GetUrl() const override;
 	virtual bool Validate() const override;
 
+#if WITH_EDITOR
+	virtual FString GetDescriptionString() const override;
+#endif //WITH_EDITOR
+	
 public:
 
 	//~ IMediaOptions interface
@@ -91,6 +95,7 @@ private:
 
 	mutable bool bUrlGuard;
 	mutable bool bValidateGuard;
+	mutable bool bConfigurationTextGuard;
 	mutable bool bLeafMediaSource;
 	mutable bool bMediaOptionGuard;
 
diff --git a/Engine/Plugins/Media/MediaFrameworkUtilities/Source/MediaFrameworkUtilities/Public/Profile/MediaProfile.h b/Engine/Plugins/Media/MediaFrameworkUtilities/Source/MediaFrameworkUtilities/Public/Profile/MediaProfile.h
--- a/Engine/Plugins/Media/MediaFrameworkUtilities/Source/MediaFrameworkUtilities/Public/Profile/MediaProfile.h
+++ b/Engine/Plugins/Media/MediaFrameworkUtilities/Source/MediaFrameworkUtilities/Public/Profile/MediaProfile.h
@@ -70,6 +70,27 @@ public:
 	 */
 	int32 NumMediaSources() const;
 
+	/**
+	 * Adds the specified media source to the profile's list of media sources
+	 * @param MediaSource The media source to add
+	 */
+	void AddMediaSource(UMediaSource* MediaSource);
+	
+	/**
+	 * Removes the specified media source from the profile's list of sources
+	 * @param Index The index of the media source to remove
+	 * @return true if the media source was removed
+	 */
+	bool RemoveMediaSource(int32 Index);
+
+	/**
+	 * Moves the media source at the specified index to the destination index, shifting other media sources down to accommodate
+	 * @param CurrentIndex The current index of the media source to move
+	 * @param DestIndex The index to put the media source at
+	 * @return true if the move occurred, otherwise false
+	 */
+	bool MoveMediaSource(int32 CurrentIndex, int32 DestIndex);
+	
 	/**
 	 * Get the media output for the selected proxy.
 	 *
@@ -82,6 +103,27 @@ public:
 	 */
 	int32 NumMediaOutputs() const;
 
+	/**
+	 * Adds the specified media output to the profile's list of media outputs
+	 * @param MediaOutput The media output to add
+	 */
+	void AddMediaOutput(UMediaOutput* MediaOutput);
+	
+	/**
+	 * Removes the specified media output from the profile's list of outputs
+	 * @param Index The index of the media output to remove
+	 * @return true if the media output was removed
+	 */
+	bool RemoveMediaOutput(int32 Index);
+
+	/**
+	 * Moves the media output at the specified index to the destination index, shifting other media outputs down to accommodate
+	 * @param CurrentIndex The current index of the media output to move
+	 * @param DestIndex The index to put the media output at
+	 * @return true if the move occurred, otherwise false
+	 */
+	bool MoveMediaOutput(int32 CurrentIndex, int32 DestIndex);
+	
 	/**
 	 * Get the timecode provider.
 	 *
diff --git a/Engine/Plugins/Media/MediaFrameworkUtilities/Source/MediaFrameworkUtilities/Public/Profile/MediaProfileSettings.h b/Engine/Plugins/Media/MediaFrameworkUtilities/Source/MediaFrameworkUtilities/Public/Profile/MediaProfileSettings.h
--- a/Engine/Plugins/Media/MediaFrameworkUtilities/Source/MediaFrameworkUtilities/Public/Profile/MediaProfileSettings.h
+++ b/Engine/Plugins/Media/MediaFrameworkUtilities/Source/MediaFrameworkUtilities/Public/Profile/MediaProfileSettings.h
@@ -76,9 +76,15 @@ public:
 	/** Set the media proxies used by the engine. */
 	void SetMediaSourceProxy(const TArray<UProxyMediaSource*>& Proxies);
 
+	/** Fills the list of proxy sources with the default proxy assets stored in the media framework plugin */
+	void FillDefaultMediaSourceProxies(int32 NumSources, bool bSilent = false);
+	
 	/** Set the media proxies used by the engine. */
 	void SetMediaOutputProxy(const TArray<UProxyMediaOutput*>& Proxies);
 
+	/** Fills the list of proxy outputs with the default proxy assets stored in the media framework plugin */
+	void FillDefaultMediaOutputProxies(int32 NumOutputs, bool bSilent = false);
+	
 	virtual void PostEditChangeProperty(struct FPropertyChangedEvent& PropertyChangedEvent) override;
 
 	/** Called when a proxy changed. */
diff --git a/Engine/Plugins/Media/MediaFrameworkUtilities/Source/MediaFrameworkUtilitiesEditor/MediaFrameworkUtilitiesEditor.Build.cs b/Engine/Plugins/Media/MediaFrameworkUtilities/Source/MediaFrameworkUtilitiesEditor/MediaFrameworkUtilitiesEditor.Build.cs
--- a/Engine/Plugins/Media/MediaFrameworkUtilities/Source/MediaFrameworkUtilitiesEditor/MediaFrameworkUtilitiesEditor.Build.cs
+++ b/Engine/Plugins/Media/MediaFrameworkUtilities/Source/MediaFrameworkUtilitiesEditor/MediaFrameworkUtilitiesEditor.Build.cs
@@ -9,6 +9,7 @@ namespace UnrealBuildTool.Rules
 			PrivateDependencyModuleNames.AddRange(
 				new string[]
 				{
+					"ApplicationCore",
 					"AssetRegistry",
 					"AssetTools",
 					"ClassViewer",
@@ -35,6 +36,7 @@ namespace UnrealBuildTool.Rules
 					"SlateCore",
 					"TimeManagement",
 					"ToolMenus",
+					"ToolWidgets",
 					"UnrealEd",
 					"WorkspaceMenuStructure",
 				});
diff --git a/Engine/Plugins/Media/MediaFrameworkUtilities/Source/MediaFrameworkUtilitiesEditor/Private/AssetEditor/MediaProfileEditor.cpp b/Engine/Plugins/Media/MediaFrameworkUtilities/Source/MediaFrameworkUtilitiesEditor/Private/AssetEditor/MediaProfileEditor.cpp
new file mode 100644
--- /dev/null
+++ b/Engine/Plugins/Media/MediaFrameworkUtilities/Source/MediaFrameworkUtilitiesEditor/Private/AssetEditor/MediaProfileEditor.cpp
@@ -0,0 +1,184 @@
+﻿// Copyright Epic Games, Inc. All Rights Reserved.
+
+
+#include "MediaProfileEditor.h"
+
+#include "PropertyEditorModule.h"
+#include "SMediaProfileSourcesTreeView.h"
+#include "Modules/ModuleManager.h"
+#include "Profile/IMediaProfileManager.h"
+#include "Profile/MediaProfile.h"
+#include "Widgets/Docking/SDockTab.h"
+
+#define LOCTEXT_NAMESPACE "MediaProfileEditor"
+
+const FName FMediaProfileEditor::AppName =  FName(TEXT("MediaProfileEditor"));
+const FName FMediaProfileEditor::MediaOutputTabId =  FName(TEXT("MediaOutput"));
+const FName FMediaProfileEditor::MediaTreeTabId =  FName(TEXT("MediaTree"));
+const FName FMediaProfileEditor::DetailsTabId =  FName(TEXT("Details"));
+const FName FMediaProfileEditor::TimecodeTabId =  FName(TEXT("Timecode"));
+
+TSharedRef<FMediaProfileEditor> FMediaProfileEditor::CreateMediaProfileEditor(const EToolkitMode::Type Mode, const TSharedPtr<IToolkitHost>& InitToolkitHost, UMediaProfile* InMediaProfile)
+{
+	TSharedRef<FMediaProfileEditor> NewEditor = MakeShared<FMediaProfileEditor>();
+	NewEditor->Initialize(Mode, InitToolkitHost, InMediaProfile);
+	return NewEditor;
+}
+
+void FMediaProfileEditor::Initialize(const EToolkitMode::Type Mode, const TSharedPtr<IToolkitHost>& InitToolkitHost, UMediaProfile* InMediaProfile)
+{
+	GEditor->GetEditorSubsystem<UAssetEditorSubsystem>()->CloseOtherEditors(InMediaProfile, this);
+	MediaProfileBeingEdited = InMediaProfile;
+
+	FPropertyEditorModule& PropertyModule = FModuleManager::LoadModuleChecked<FPropertyEditorModule>("PropertyEditor");
+	FDetailsViewArgs DatabaseDetailsArgs;
+	DetailsView = PropertyModule.CreateDetailView(DatabaseDetailsArgs);
+	DetailsView->SetObject(InMediaProfile);
+	
+	const TSharedRef<FTabManager::FLayout> StandaloneDefaultLayout = FTabManager::NewLayout("Standalone_MediaProfileEditor_Layout")
+		->AddArea
+		(
+			FTabManager::NewPrimaryArea()
+			->SetOrientation(Orient_Vertical)
+			->Split
+			(
+				FTabManager::NewSplitter()
+				->SetOrientation(Orient_Horizontal)
+				->Split
+				(
+					FTabManager::NewStack()
+					->SetSizeCoefficient(0.8f)
+					->SetHideTabWell(true)
+					->AddTab(MediaOutputTabId, ETabState::OpenedTab)
+				)
+				->Split
+				(
+					FTabManager::NewSplitter()
+					->SetOrientation(Orient_Vertical)
+					->SetSizeCoefficient(0.2f)
+					->Split
+					(
+						FTabManager::NewStack()
+						->SetSizeCoefficient(0.4f)
+						->AddTab(MediaTreeTabId, ETabState::OpenedTab)
+					)
+					->Split
+					(
+						FTabManager::NewStack()
+						->SetSizeCoefficient(0.6f)
+						->AddTab(DetailsTabId, ETabState::OpenedTab)
+						->AddTab(TimecodeTabId, ETabState::OpenedTab)
+					)
+				)
+			)
+		);
+
+	const bool bCreateDefaultStandaloneMenu = true;
+	const bool bCreateDefaultToolbar = true;
+	InitAssetEditor(Mode, InitToolkitHost, AppName, StandaloneDefaultLayout, bCreateDefaultStandaloneMenu, bCreateDefaultToolbar, InMediaProfile);
+
+	// TODO: Not sure this is the best place for this, but for now, when we open a media profile, automatically make it the active media profile
+	IMediaProfileManager::Get().SetCurrentMediaProfile(InMediaProfile);
+}
+
+void FMediaProfileEditor::RegisterTabSpawners(const TSharedRef<class FTabManager>& InTabManager)
+{
+	WorkspaceMenuCategory = InTabManager->AddLocalWorkspaceMenuCategory(LOCTEXT("WorkspaceMenuCategory", "Media Profile Editor"));
+	
+	FAssetEditorToolkit::RegisterTabSpawners(InTabManager);
+
+	InTabManager->RegisterTabSpawner(MediaOutputTabId, FOnSpawnTab::CreateSP(this, &FMediaProfileEditor::SpawnTab_MediaOutput))
+		.SetDisplayName(LOCTEXT("MediaOutputTabDisplayName", "Media Output"))
+		.SetGroup(WorkspaceMenuCategory.ToSharedRef());
+
+	InTabManager->RegisterTabSpawner(MediaTreeTabId, FOnSpawnTab::CreateSP(this, &FMediaProfileEditor::SpawnTab_MediaTree))
+		.SetDisplayName(LOCTEXT("MediaTreeTabDisplayName", "Media"))
+		.SetGroup(WorkspaceMenuCategory.ToSharedRef());
+
+	InTabManager->RegisterTabSpawner(DetailsTabId, FOnSpawnTab::CreateSP(this, &FMediaProfileEditor::SpawnTab_Details))
+		.SetDisplayName(LOCTEXT("DetailsTabDisplayName", "Details"))
+		.SetGroup(WorkspaceMenuCategory.ToSharedRef())
+		.SetIcon(FSlateIcon(FAppStyle::GetAppStyleSetName(), "LevelEditor.Tabs.Details"));
+
+	InTabManager->RegisterTabSpawner(TimecodeTabId, FOnSpawnTab::CreateSP(this, &FMediaProfileEditor::SpawnTab_Timecode))
+		.SetDisplayName(LOCTEXT("TimecodeTabDisplayName", "Timecode/Genlock"))
+		.SetGroup(WorkspaceMenuCategory.ToSharedRef());
+}
+
+void FMediaProfileEditor::UnregisterTabSpawners(const TSharedRef<class FTabManager>& InTabManager)
+{
+	FAssetEditorToolkit::UnregisterTabSpawners(InTabManager);
+
+	InTabManager->UnregisterTabSpawner(MediaOutputTabId);
+	InTabManager->UnregisterTabSpawner(MediaTreeTabId);
+	InTabManager->UnregisterTabSpawner(DetailsTabId);
+	InTabManager->UnregisterTabSpawner(TimecodeTabId);
+}
+
+FName FMediaProfileEditor::GetToolkitFName() const
+{
+	return FName("MediaProfileEditor");
+}
+
+FText FMediaProfileEditor::GetBaseToolkitName() const
+{
+	return LOCTEXT("MediaProfileEditorLabel", "Media Profile Editor");
+}
+
+FText FMediaProfileEditor::GetToolkitName() const
+{
+	return FText::FromString(MediaProfileBeingEdited->GetName());
+}
+
+FText FMediaProfileEditor::GetToolkitToolTipText() const
+{
+	return FAssetEditorToolkit::GetToolTipTextForObject(MediaProfileBeingEdited);
+}
+
+FString FMediaProfileEditor::GetWorldCentricTabPrefix() const
+{
+	return TEXT("MediaProfileEditor");
+}
+
+FLinearColor FMediaProfileEditor::GetWorldCentricTabColorScale() const
+{
+	return FLinearColor::White;
+}
+
+TSharedRef<SDockTab> FMediaProfileEditor::SpawnTab_MediaOutput(const FSpawnTabArgs& InArgs)
+{
+	return SNew(SDockTab)
+		.Label(LOCTEXT("MediaOutputTabLabel", "Media Output"))
+		[
+			SNullWidget::NullWidget
+		];
+}
+
+TSharedRef<SDockTab> FMediaProfileEditor::SpawnTab_MediaTree(const FSpawnTabArgs& InArgs)
+{
+	return SNew(SDockTab)
+		.Label(LOCTEXT("MediaTreeTabLabel", "Media"))
+		[
+			SNew(SMediaProfileSourcesTreeView, MediaProfileBeingEdited)
+		];
+}
+
+TSharedRef<SDockTab> FMediaProfileEditor::SpawnTab_Details(const FSpawnTabArgs& InArgs)
+{
+	return SNew(SDockTab)
+		.Label(LOCTEXT("DetailsTabLabel", "Details"))
+		[
+			DetailsView.ToSharedRef()
+		];
+}
+
+TSharedRef<SDockTab> FMediaProfileEditor::SpawnTab_Timecode(const FSpawnTabArgs& InArgs)
+{
+	return SNew(SDockTab)
+		.Label(LOCTEXT("TimecodeTabLabel", "Timecode/Genlock"))
+		[
+			SNullWidget::NullWidget
+		];
+}
+
+#undef LOCTEXT_NAMESPACE
diff --git a/Engine/Plugins/Media/MediaFrameworkUtilities/Source/MediaFrameworkUtilitiesEditor/Private/AssetEditor/MediaProfileEditor.h b/Engine/Plugins/Media/MediaFrameworkUtilities/Source/MediaFrameworkUtilitiesEditor/Private/AssetEditor/MediaProfileEditor.h
new file mode 100644
--- /dev/null
+++ b/Engine/Plugins/Media/MediaFrameworkUtilities/Source/MediaFrameworkUtilitiesEditor/Private/AssetEditor/MediaProfileEditor.h
@@ -0,0 +1,52 @@
+﻿// Copyright Epic Games, Inc. All Rights Reserved.
+
+#pragma once
+
+#include "CoreMinimal.h"
+#include "Toolkits/AssetEditorToolkit.h"
+
+class UMediaProfile;
+
+/**
+ * Media profile editor 2.0, which replaces the old, simple editor with a more robust one
+ */
+class FMediaProfileEditor : public FAssetEditorToolkit
+{
+private:
+	static const FName AppName;
+	static const FName MediaOutputTabId;
+	static const FName MediaTreeTabId;
+	static const FName DetailsTabId;
+	static const FName TimecodeTabId;
+	
+public:
+	static TSharedRef<FMediaProfileEditor> CreateMediaProfileEditor(const EToolkitMode::Type Mode, const TSharedPtr< IToolkitHost >& InitToolkitHost, UMediaProfile* InMediaProfile);
+
+	void Initialize(const EToolkitMode::Type Mode, const TSharedPtr< IToolkitHost >& InitToolkitHost, UMediaProfile* InMediaProfile);
+
+	// IToolkit interface
+	virtual void RegisterTabSpawners(const TSharedRef<class FTabManager>& InTabManager) override;
+	virtual void UnregisterTabSpawners(const TSharedRef<class FTabManager>& InTabManager) override;
+	// End of IToolkit interface
+
+	// FAssetEditorToolkit
+	virtual FName GetToolkitFName() const override;
+	virtual FText GetBaseToolkitName() const override;
+	virtual FText GetToolkitName() const override;
+	virtual FText GetToolkitToolTipText() const override;
+	virtual FString GetWorldCentricTabPrefix() const override;
+	virtual FLinearColor GetWorldCentricTabColorScale() const override;
+	// End of FAssetEditorToolkit
+
+private:
+	TSharedRef<SDockTab> SpawnTab_MediaOutput(const FSpawnTabArgs& InArgs);
+	TSharedRef<SDockTab> SpawnTab_MediaTree(const FSpawnTabArgs& InArgs);
+	TSharedRef<SDockTab> SpawnTab_Details(const FSpawnTabArgs& InArgs);
+	TSharedRef<SDockTab> SpawnTab_Timecode(const FSpawnTabArgs& InArgs);
+private:
+	/** The media profile being edited by this editor */
+	TObjectPtr<UMediaProfile> MediaProfileBeingEdited;
+
+	/** Details view to display the media profile properties */
+	TSharedPtr<IDetailsView> DetailsView;
+};
diff --git a/Engine/Plugins/Media/MediaFrameworkUtilities/Source/MediaFrameworkUtilitiesEditor/Private/AssetEditor/SMediaProfileSourcesTreeView.cpp b/Engine/Plugins/Media/MediaFrameworkUtilities/Source/MediaFrameworkUtilitiesEditor/Private/AssetEditor/SMediaProfileSourcesTreeView.cpp
new file mode 100644
--- /dev/null
+++ b/Engine/Plugins/Media/MediaFrameworkUtilities/Source/MediaFrameworkUtilitiesEditor/Private/AssetEditor/SMediaProfileSourcesTreeView.cpp
@@ -0,0 +1,1202 @@
+﻿// Copyright Epic Games, Inc. All Rights Reserved.
+
+#include "SMediaProfileSourcesTreeView.h"
+
+#include "Editor.h"
+#include "MediaOutput.h"
+#include "MediaSource.h"
+#include "ScopedTransaction.h"
+#include "SlateOptMacros.h"
+#include "SPositiveActionButton.h"
+#include "SSettingsEditorCheckoutNotice.h"
+#include "UnrealExporter.h"
+#include "DragAndDrop/DecoratedDragDropOp.h"
+#include "Exporters/Exporter.h"
+#include "Filters/GenericFilter.h"
+#include "Filters/SBasicFilterBar.h"
+#include "Filters/SFilterBar.h"
+#include "Filters/SFilterSearchBox.h"
+#include "Framework/Commands/GenericCommands.h"
+#include "MediaAssets/ProxyMediaOutput.h"
+#include "MediaAssets/ProxyMediaSource.h"
+#include "Misc/MessageDialog.h"
+#include "Profile/MediaProfile.h"
+#include "Profile/MediaProfileSettings.h"
+#include "Styling/SlateIconFinder.h"
+#include "Widgets/Input/SComboButton.h"
+
+#define LOCTEXT_NAMESPACE "MediaProfileSourcesTreeView"
+
+const FLazyName SMediaProfileSourcesTreeView::Column_Index(TEXT("Column_Index"));
+const FLazyName SMediaProfileSourcesTreeView::Column_ItemLabel(TEXT("Column_ItemLabel"));
+const FLazyName SMediaProfileSourcesTreeView::Column_ItemType(TEXT("Column_ItemType"));
+const FLazyName SMediaProfileSourcesTreeView::Column_Configuration(TEXT("Column_Configuration"));
+
+namespace MediaSourcesTreeView
+{
+	/** Attempts to saves the media profile settings to the project's config files */
+	void TrySaveMediaProfileSettings()
+	{
+		// This will log a warning if the default file was read-only
+		const bool bUpdated = GetMutableDefault<UMediaProfileSettings>()->TryUpdateDefaultConfigFile();
+
+		// If the default settings could not be updated, try once to make the default config file writable, then try updating it again
+		if (!bUpdated)
+		{
+			const FString RelativeConfigFilePath = GetMutableDefault<UMediaProfileSettings>()->GetDefaultConfigFilename();
+			const FString TargetFilePath = FPaths::ConvertRelativePathToFull(RelativeConfigFilePath);
+
+			constexpr bool bForceSourceControlUpdate = false;
+			constexpr bool bShowErrorInNotification = true;
+			if (SettingsHelpers::CheckOutOrAddFile(TargetFilePath, bForceSourceControlUpdate, bShowErrorInNotification))
+			{
+				GetMutableDefault<UMediaProfileSettings>()->TryUpdateDefaultConfigFile();
+			}
+		}
+	}
+}
+
+/** Slate widget for the individual rows of the media tree view */
+class SMediaTreeItemRow : public SMultiColumnTableRow<SMediaProfileSourcesTreeView::FMediaTreeItemPtr>
+{
+private:
+	/** Drag/drop operation for changing the order of sources or outputs in the tree view */
+	class FMediaDragDropOp : public FDecoratedDragDropOp
+	{
+	public:
+		DRAG_DROP_OPERATOR_TYPE(FMediaDragDropOp, FDecoratedDragDropOp)
+
+		static TSharedRef<FMediaDragDropOp> New(const SMediaProfileSourcesTreeView::FMediaTreeItemPtr& InDraggedItem)
+		{
+			TSharedRef<FMediaDragDropOp> Operation = MakeShared<FMediaDragDropOp>();
+			Operation->DraggedItem = InDraggedItem;
+			Operation->Construct();
+			return Operation;
+		}
+
+		virtual void Construct() override
+		{
+			if (DraggedItem.IsValid())
+			{
+				SetToolTip(DraggedItem.Pin()->Label, DraggedItem.Pin()->Icon);
+			}
+			
+			FDecoratedDragDropOp::Construct();
+		}
+
+		TWeakPtr<SMediaProfileSourcesTreeView::FMediaTreeItem> DraggedItem;
+	};
+	
+public:
+	/** Delegate that is called when a media item is moved. Parameters are the item being moved, and the index to move it to  */
+	DECLARE_DELEGATE_TwoParams(FOnMoveMediaItem, SMediaProfileSourcesTreeView::FMediaTreeItemPtr, int32)
+	
+	SLATE_BEGIN_ARGS(SMediaTreeItemRow) {}
+		SLATE_EVENT(FOnMoveMediaItem, OnMoveMediaItem)
+	SLATE_END_ARGS()
+
+	void Construct(const FArguments& InArgs, const TSharedRef<STableViewBase>& InOwnerTableView, SMediaProfileSourcesTreeView::FMediaTreeItemPtr InTreeItem)
+	{
+		TreeItem = InTreeItem;
+		OnMoveMediaItem = InArgs._OnMoveMediaItem;
+		
+		STableRow<SMediaProfileSourcesTreeView::FMediaTreeItemPtr>::FArguments Args = FSuperRowType::FArguments()
+			.Style(&FAppStyle::Get().GetWidgetStyle<FTableRowStyle>("SceneOutliner.TableViewRow"))
+			.OnDragDetected(this, &SMediaTreeItemRow::HandleDragDetected)
+			.OnCanAcceptDrop(this, &SMediaTreeItemRow::HandleCanAcceptDrop)
+			.OnAcceptDrop(this, &SMediaTreeItemRow::HandleAcceptDrop);
+		
+		SMultiColumnTableRow<SMediaProfileSourcesTreeView::FMediaTreeItemPtr>::Construct(Args, InOwnerTableView);
+	}
+
+	virtual TSharedRef<SWidget> GenerateWidgetForColumn( const FName& ColumnName ) override
+	{
+		SMediaProfileSourcesTreeView::FMediaTreeItemPtr PinnedTreeItem = TreeItem.Pin();
+		if (!PinnedTreeItem.IsValid())
+		{
+			return SNullWidget::NullWidget;
+		}
+
+		auto CreateColumnWidget = [](const FText& InLabel, const FSlateBrush* InIcon)
+		{
+			TSharedPtr<SWidget> LabelWidget = SNullWidget::NullWidget;
+			if (!InLabel.IsEmptyOrWhitespace())
+			{
+				LabelWidget = SNew(STextBlock).Text(InLabel);
+			}
+
+			TSharedPtr<SWidget> IconWidget = SNullWidget::NullWidget;
+			if (InIcon)
+			{
+				IconWidget = SNew(SBox)
+					.WidthOverride(16.0f)
+					.HeightOverride(16.0f)
+					[
+						SNew(SImage).Image(InIcon)
+					];
+			}
+
+			const FMargin LeftMargin = InIcon ? FMargin(6.0f, 1.0f, 6.0f, 1.0f) : FMargin();
+			return SNew(SHorizontalBox)
+				+SHorizontalBox::Slot()
+				.HAlign(HAlign_Center)
+				.VAlign(VAlign_Center)
+				.Padding(LeftMargin)
+				.AutoWidth()
+				[
+					IconWidget.ToSharedRef()
+				]
+
+				+SHorizontalBox::Slot()
+				.FillWidth(1.0f)
+				[
+					LabelWidget.ToSharedRef()
+				];
+		};
+		
+		if (ColumnName == SMediaProfileSourcesTreeView::Column_Index)
+		{
+			if (PinnedTreeItem->IsLeaf())
+			{
+				return CreateColumnWidget(FText::Format(LOCTEXT("MediaSourceIndexLabelFormat", "{0}"), PinnedTreeItem->Index + 1), nullptr);
+			}
+		}
+
+		if (ColumnName == SMediaProfileSourcesTreeView::Column_ItemLabel)
+		{
+			return SNew(SBox)
+				.MinDesiredHeight(20.0f)
+				[
+					SNew(SHorizontalBox)
+
+					+SHorizontalBox::Slot()
+					.AutoWidth()
+					.Padding(6.f, 0.f, 0.f, 0.f)
+					[
+						SNew(SExpanderArrow, SharedThis(this)).IndentAmount(12)
+					]
+
+					+SHorizontalBox::Slot()
+					.FillWidth(1.0f)
+					[
+						CreateColumnWidget(PinnedTreeItem->Label, PinnedTreeItem->Icon)
+					]
+				];
+		}
+
+		if (ColumnName == SMediaProfileSourcesTreeView::Column_ItemType)
+		{
+			if (PinnedTreeItem->IsLeaf())
+			{
+				return CreateColumnWidget(PinnedTreeItem->Type, PinnedTreeItem->Icon);
+			}
+		}
+
+		if (ColumnName == SMediaProfileSourcesTreeView::Column_Configuration)
+		{
+			if (PinnedTreeItem->IsLeaf())
+			{
+				return CreateColumnWidget(PinnedTreeItem->Configuration, nullptr);
+			}
+		}
+		
+		return SNullWidget::NullWidget;
+	}
+
+private:
+	FReply HandleDragDetected(const FGeometry& MyGeometry, const FPointerEvent& MouseEvent)
+	{
+		if (!TreeItem.IsValid())
+		{
+			return FReply::Unhandled();
+		}
+
+		if (!TreeItem.Pin()->IsLeaf())
+		{
+			return FReply::Unhandled();
+		}
+		
+		if (MouseEvent.IsMouseButtonDown(EKeys::LeftMouseButton))
+		{
+			TSharedPtr<FDragDropOperation> DragDropOp = FMediaDragDropOp::New(TreeItem.Pin());
+			return FReply::Handled().BeginDragDrop(DragDropOp.ToSharedRef());
+		}
+
+		return FReply::Unhandled();
+	}
+
+	TOptional<EItemDropZone> HandleCanAcceptDrop(const FDragDropEvent& DragDropEvent, EItemDropZone InItemDropZone, SMediaProfileSourcesTreeView::FMediaTreeItemPtr InTreeItem)
+	{
+		TSharedPtr<FMediaDragDropOp> DragDropOp = DragDropEvent.GetOperationAs<FMediaDragDropOp>();
+
+		if (!DragDropOp.IsValid() || !DragDropOp->DraggedItem.IsValid())
+		{
+			return TOptional<EItemDropZone>();
+		}
+
+		SMediaProfileSourcesTreeView::FMediaTreeItemPtr DraggedItem = DragDropOp->DraggedItem.Pin();
+		if (DraggedItem->BackingClass == InTreeItem->BackingClass)
+		{
+			if (!InTreeItem->IsLeaf())
+			{
+				return EItemDropZone::BelowItem;
+			}
+
+			return InItemDropZone == EItemDropZone::OntoItem ? EItemDropZone::AboveItem : InItemDropZone;
+		}
+		
+		return TOptional<EItemDropZone>();
+	}
+
+	
+	FReply HandleAcceptDrop(const FDragDropEvent& DragDropEvent, EItemDropZone InItemDropZone, SMediaProfileSourcesTreeView::FMediaTreeItemPtr InTreeItem)
+	{
+		TSharedPtr<FMediaDragDropOp> DragDropOp = DragDropEvent.GetOperationAs<FMediaDragDropOp>();
+
+		if (!DragDropOp.IsValid() || !DragDropOp->DraggedItem.IsValid())
+		{
+			return FReply::Unhandled();
+		}
+
+		if (OnMoveMediaItem.IsBound())
+		{
+			const int32 DestIndex = InItemDropZone == EItemDropZone::AboveItem ? InTreeItem->Index: InTreeItem->Index + 1;
+			OnMoveMediaItem.Execute(DragDropOp->DraggedItem.Pin(), DestIndex);
+			return FReply::Handled();
+		}
+		
+		return FReply::Unhandled();
+	}
+	
+private:
+	TWeakPtr<SMediaProfileSourcesTreeView::FMediaTreeItem> TreeItem;
+
+	FOnMoveMediaItem OnMoveMediaItem;
+};
+
+/**
+ * Widget for the filter bar, needs to be a subclass that overrides MakeAddFilterMenu to give the filter bar its own unique menu name,
+ * otherwise the editor will get confused with any other basic filter bar used elsewhere.
+ */
+template<typename TFilterType>
+class SMediaSourcesFilterBar : public SBasicFilterBar<TFilterType>
+{
+	using Super = SBasicFilterBar<TFilterType>;
+	
+public:
+
+	using FOnFilterChanged = typename SBasicFilterBar<TFilterType>::FOnFilterChanged;
+	using FCreateTextFilter = typename SBasicFilterBar<TFilterType>::FCreateTextFilter;
+
+	SLATE_BEGIN_ARGS(SMediaSourcesFilterBar)
+	{}
+		SLATE_EVENT(SMediaSourcesFilterBar<TFilterType>::FOnFilterChanged, OnFilterChanged)
+		SLATE_ARGUMENT(TArray<TSharedRef<FFilterBase<TFilterType>>>, CustomFilters)	
+	SLATE_END_ARGS()
+
+	void Construct(const FArguments& InArgs)
+	{
+		typename SBasicFilterBar<TFilterType>::FArguments Args;
+		Args._OnFilterChanged = InArgs._OnFilterChanged;
+		Args._CustomFilters = InArgs._CustomFilters;
+		Args._UseSectionsForCategories = true;
+		
+		SBasicFilterBar<TFilterType>::Construct(Args.FilterPillStyle(EFilterPillStyle::Basic));
+	}
+
+private:
+	virtual TSharedRef<SWidget> MakeAddFilterMenu() override
+	{
+		const FName FilterMenuName = "MediaSourcesFilterBar.FilterMenu";
+		if (!UToolMenus::Get()->IsMenuRegistered(FilterMenuName))
+		{
+			UToolMenu* Menu = UToolMenus::Get()->RegisterMenu(FilterMenuName);
+			Menu->bShouldCloseWindowAfterMenuSelection = true;
+			Menu->bCloseSelfOnly = true;
+
+			Menu->AddDynamicSection(NAME_None, FNewToolMenuDelegate::CreateLambda([](UToolMenu* InMenu)
+			{
+				if (UFilterBarContext* Context = InMenu->FindContext<UFilterBarContext>())
+				{
+					Context->PopulateFilterMenu.ExecuteIfBound(InMenu);
+					Context->OnExtendAddFilterMenu.ExecuteIfBound(InMenu);
+				}
+			}));
+		}
+
+		UFilterBarContext* FilterBarContext = NewObject<UFilterBarContext>();
+		FilterBarContext->PopulateFilterMenu = FOnPopulateAddFilterMenu::CreateLambda([this](UToolMenu* Menu)
+		{
+			Super::PopulateCommonFilterSections(Menu);
+			Super::PopulateCustomFilters(Menu);
+		});
+		FToolMenuContext ToolMenuContext(FilterBarContext);
+		
+		return UToolMenus::Get()->GenerateWidget(FilterMenuName, ToolMenuContext);
+	}
+};
+
+/** Toolbar widget that contains filtering, the add button, and the settings button */
+class SMediaProfileSourcesToolbar : public SCompoundWidget
+{
+private:
+	typedef SMediaProfileSourcesTreeView::FMediaTreeItemPtr FFilterType;
+	typedef TTextFilter<FFilterType> FMediaTreeFilter;
+	
+public:
+	DECLARE_DELEGATE(FOnFilterChanged)
+	DECLARE_DELEGATE(FOnMediaSourceAdded)
+	DECLARE_DELEGATE(FOnMediaOutputAdded)
+	
+	SLATE_BEGIN_ARGS(SMediaProfileSourcesToolbar) {}
+		SLATE_EVENT(FOnFilterChanged, OnFilterChanged)
+		SLATE_EVENT(FOnMediaSourceAdded, OnMediaSourceAdded)
+		SLATE_EVENT(FOnMediaOutputAdded, OnMediaOutputAdded)
+	SLATE_END_ARGS()
+
+	void Construct(const FArguments& InArgs, UMediaProfile* InMediaProfile)
+	{
+		MediaProfile = InMediaProfile;
+		OnFilterChanged = InArgs._OnFilterChanged;
+		OnMediaSourceAdded = InArgs._OnMediaSourceAdded;
+		OnMediaOutputAdded = InArgs._OnMediaOutputAdded;
+
+		InitializeFilters();
+		
+		FilterBar = SNew(SMediaSourcesFilterBar<FFilterType>)
+			.CustomFilters(CustomFilters)
+			.OnFilterChanged(this, &SMediaProfileSourcesToolbar::FilterChanged);
+		
+		ChildSlot
+		[
+			SNew(SVerticalBox)
+
+			+SVerticalBox::Slot()
+			.AutoHeight()
+			[
+				SNew(SHorizontalBox)
+
+				+SHorizontalBox::Slot()
+				.AutoWidth()
+				.Padding(0.0f, 0.0f, 4.0f, 0.0f)
+				[
+					SNew(SPositiveActionButton)
+					.Icon(FAppStyle::Get().GetBrush("Icons.Plus"))
+					.Text(LOCTEXT("AddMediaButtonLabel", "Add"))
+					.ToolTipText(LOCTEXT("AddMediaButtonToolTip", "Add new media source or media output"))
+					.OnGetMenuContent(this, &SMediaProfileSourcesToolbar::GetAddMediaMenuContent)
+				]
+
+				+SHorizontalBox::Slot()
+				.VAlign(VAlign_Center)
+				.Padding(0.0f, 0.0f, 2.0f, 0.0f)
+				.AutoWidth()
+				[
+					SMediaSourcesFilterBar<FFilterType>::MakeAddFilterButton(FilterBar.ToSharedRef())
+				]
+				
+				+SHorizontalBox::Slot()
+				.FillWidth(1.0f)
+				[
+					SNew(SFilterSearchBox)
+					.HintText(LOCTEXT("FilterSearch", "Search..."))
+					.ToolTipText(LOCTEXT("FilterSearchHint", "Type here to search for specific media sources or outputs"))
+					.OnTextChanged_Lambda([this](const FText& InText)
+					{
+						MediaTreeTextFilter->SetRawFilterText(InText);
+					})
+				]
+
+				+SHorizontalBox::Slot()
+				.AutoWidth()
+				.Padding(4.f, 0.f, 0.f, 0.f)
+				[
+					SNew(SComboButton)
+					.ComboButtonStyle(FAppStyle::Get(), "SimpleComboButtonWithIcon")
+					//.OnGetMenuContent( this, &SSceneOutliner::GetViewButtonContent, Mode->ShowFilterOptions())
+					.HasDownArrow(false)
+					.ButtonContent()
+					[
+						SNew(SImage)
+						.ColorAndOpacity(FSlateColor::UseForeground())
+						.Image(FAppStyle::Get().GetBrush("Icons.Settings"))
+					]
+				]
+			]
+
+			+SVerticalBox::Slot()
+			.AutoHeight()
+			[
+				FilterBar.ToSharedRef()
+			]
+		];
+	}
+
+	bool ItemPassesFilters(const FFilterType& InItem) const
+	{
+		// Performs an OR check on any of the filters that are enabled in the filter bar, mirroring the scene outliner filter bar behavior
+		auto PassesAnyFilterBarFilter = [this](const FFilterType& InItem)
+		{
+			TSharedPtr<TFilterCollection<FFilterType>> FilterCollection = FilterBar->GetAllActiveFilters();
+			const int32 NumFilters = FilterCollection->Num();
+			if (NumFilters == 0)
+			{
+				return true;
+			}
+		
+			for (int32 Index = 0; Index < NumFilters; ++Index)
+			{
+				if (FilterCollection->GetFilterAtIndex(Index)->PassesFilter(InItem))
+				{
+					return true;
+				}
+			}
+
+			return false;
+		};
+		
+		return MediaTreeTextFilter->PassesFilter(InItem) && PassesAnyFilterBarFilter(InItem);
+	}
+	
+private:
+	/** Creates all the filter bar filters and initializes the text filter */
+	void InitializeFilters()
+	{
+		// Basic filters, for filtering out the sources or outputs broadly
+		{
+			TSharedPtr<FFilterCategory> BasicFiltersCategory = MakeShared<FFilterCategory>(LOCTEXT("BasicFiltersCategory", "Basic"), FText::GetEmpty());
+
+			TSharedPtr<FGenericFilter<FFilterType>> MediaSourcesFilter = MakeShared<FGenericFilter<FFilterType>>(
+				BasicFiltersCategory,
+				TEXT("MediaSourcesFilter"),
+				LOCTEXT("MediaSourcesFilterName", "Sources"),
+				FGenericFilter<FFilterType>::FOnItemFiltered::CreateLambda([](FFilterType InItem)
+				{
+					if (!InItem.IsValid())
+					{
+						return false;
+					}
+
+					return InItem->BackingObject.IsValid() && InItem->BackingObject->IsA<UMediaSource>();
+				}));
+		
+			MediaSourcesFilter->SetToolTipText(LOCTEXT("MediaSourcesFilter", "Only show media sources"));
+			CustomFilters.Add(MediaSourcesFilter.ToSharedRef());
+
+			TSharedPtr<FGenericFilter<FFilterType>> MediaOutputsFilter = MakeShared<FGenericFilter<FFilterType>>(
+				BasicFiltersCategory,
+				TEXT("MediaOutputsFilter"),
+				LOCTEXT("MediaOutputsFilterName", "Outputs"),
+				FGenericFilter<FFilterType>::FOnItemFiltered::CreateLambda([](FFilterType InItem)
+				{
+					if (!InItem.IsValid())
+					{
+						return false;
+					}
+
+					return InItem->BackingObject.IsValid() && InItem->BackingObject->IsA<UMediaOutput>();
+				}));
+		
+			MediaOutputsFilter->SetToolTipText(LOCTEXT("MediaOutputsFilter", "Only show media outputs"));
+			CustomFilters.Add(MediaOutputsFilter.ToSharedRef());
+		}
+
+		// Class filters for media sources
+		{
+			TSharedPtr<FFilterCategory> SourceFiltersCategory = MakeShared<FFilterCategory>(LOCTEXT("SourceFiltersCategory", "Media Sources"), FText::GetEmpty());
+			
+			TArray<UClass*> ChildClasses;
+			GetDerivedClasses(UMediaSource::StaticClass(), ChildClasses);
+
+			for (UClass* ChildClass : ChildClasses)
+			{
+				if (ChildClass->HasAnyClassFlags(CLASS_Abstract))
+				{
+					continue;
+				}
+				
+				TSharedPtr<FGenericFilter<FFilterType>> MediaSourceSubclassFilter = MakeShared<FGenericFilter<FFilterType>>(
+				SourceFiltersCategory,
+				FString::Format(TEXT("{0}ClassFilter"), { ChildClass->GetName() }),
+				FText::Format(LOCTEXT("MediaSourceSubclassFilterName", "{0}"), ChildClass->GetDisplayNameText()),
+				FGenericFilter<FFilterType>::FOnItemFiltered::CreateLambda([ChildClass](FFilterType InItem)
+				{
+					if (!InItem.IsValid())
+					{
+						return false;
+					}
+
+					return InItem->BackingObject.IsValid() && InItem->BackingObject->IsA(ChildClass);
+				}));
+
+				CustomFilters.Add(MediaSourceSubclassFilter.ToSharedRef());
+			}
+		}
+
+		// Class filters for media outputs
+		{
+			TSharedPtr<FFilterCategory> OutputFiltersCategory = MakeShared<FFilterCategory>(LOCTEXT("OutputFiltersCategory", "Media Outputs"), FText::GetEmpty());
+			
+			TArray<UClass*> ChildClasses;
+			GetDerivedClasses(UMediaOutput::StaticClass(), ChildClasses);
+
+			for (UClass* ChildClass : ChildClasses)
+			{
+				if (ChildClass->HasAnyClassFlags(CLASS_Abstract))
+				{
+					continue;
+				}
+				
+				TSharedPtr<FGenericFilter<FFilterType>> MediaOutputSubclassFilter = MakeShared<FGenericFilter<FFilterType>>(
+				OutputFiltersCategory,
+				FString::Format(TEXT("{0}ClassFilter"), { ChildClass->GetName() }),
+				FText::Format(LOCTEXT("MediaOutputSubclassFilterName", "{0}"), ChildClass->GetDisplayNameText()),
+				FGenericFilter<FFilterType>::FOnItemFiltered::CreateLambda([ChildClass](FFilterType InItem)
+				{
+					if (!InItem.IsValid())
+					{
+						return false;
+					}
+
+					return InItem->BackingObject.IsValid() && InItem->BackingObject->IsA(ChildClass);
+				}));
+
+				CustomFilters.Add(MediaOutputSubclassFilter.ToSharedRef());
+			}
+		}
+		
+		MediaTreeTextFilter = MakeShared<FMediaTreeFilter>(FMediaTreeFilter::FItemToStringArray::CreateLambda([](const FFilterType& InItem, OUT TArray<FString>& OutStrings)
+		{
+			if (InItem.IsValid())
+			{
+				OutStrings.Add(InItem->Label.ToString());
+				OutStrings.Add(InItem->Type.ToString());
+				OutStrings.Add(InItem->Configuration.ToString());
+			}
+		}));
+		MediaTreeTextFilter->OnChanged().AddSP(this, &SMediaProfileSourcesToolbar::FilterChanged);
+	}
+	
+	TSharedRef<SWidget> GetAddMediaMenuContent()
+	{
+		constexpr bool bCloseMenuAfterSelection = true;
+		FMenuBuilder MenuBuilder(bCloseMenuAfterSelection, nullptr);
+
+		MenuBuilder.AddMenuEntry(
+			LOCTEXT("AddMediaSourceMenuLabel", "Add Media Source"),
+			LOCTEXT("AddMediaSourceMenuToolTip", "Add a new media source to the profile"),
+			FSlateIcon(),
+			FUIAction(FExecuteAction::CreateLambda([this]() { OnMediaSourceAdded.ExecuteIfBound(); })));
+
+		MenuBuilder.AddMenuEntry(
+			LOCTEXT("AddMediaOutputMenuLabel", "Add Media Output"),
+			LOCTEXT("AddMediaOutputMenuToolTip", "Add a new media output to the profile"),
+			FSlateIcon(),
+			FUIAction(FExecuteAction::CreateLambda([this]() { OnMediaOutputAdded.ExecuteIfBound(); })));
+		
+		return MenuBuilder.MakeWidget();
+	}
+
+	void FilterChanged()
+	{
+		OnFilterChanged.ExecuteIfBound();
+	}
+	
+private:
+	TWeakObjectPtr<UMediaProfile> MediaProfile;
+	TArray<TSharedRef<FFilterBase<FFilterType>>> CustomFilters;
+	TSharedPtr<FMediaTreeFilter> MediaTreeTextFilter;
+	
+	TSharedPtr<SMediaSourcesFilterBar<FFilterType>> FilterBar;
+	
+	FOnFilterChanged OnFilterChanged;
+	FOnMediaSourceAdded OnMediaSourceAdded;
+	FOnMediaOutputAdded OnMediaOutputAdded;
+};
+
+void SMediaProfileSourcesTreeView::Construct(const FArguments& InArgs, UMediaProfile* InMediaProfile)
+{
+	MediaProfile = InMediaProfile;
+
+	CommandList = MakeShared<FUICommandList>();
+	BindCommands();
+	
+	FillMediaTree();
+	
+	ChildSlot
+	[
+		SNew(SVerticalBox)
+
+		+SVerticalBox::Slot()
+		.AutoHeight()
+		.Padding(8.0f, 8.0f, 8.0f, 4.0f)
+		[
+			SAssignNew(Toolbar, SMediaProfileSourcesToolbar, InMediaProfile)
+			.OnFilterChanged(this, &SMediaProfileSourcesTreeView::OnFilterChanged)
+			.OnMediaSourceAdded(this, &SMediaProfileSourcesTreeView::AddNewMediaSource)
+			.OnMediaOutputAdded(this, &SMediaProfileSourcesTreeView::AddNewMediaOutput)
+		]
+		
+		+SVerticalBox::Slot()
+		.FillHeight(1.0f)
+		[
+			SAssignNew(MediaTreeView, STreeView<FMediaTreeItemPtr>)
+			.TreeItemsSource(&MediaTreeItems)
+			.HeaderRow
+			(
+				SNew(SHeaderRow)
+
+				+SHeaderRow::Column(Column_Index)
+				.DefaultLabel(LOCTEXT("IndexColumnLabel", "#"))
+				.ManualWidth(20.0f)
+
+				+SHeaderRow::Column(Column_ItemLabel)
+				.DefaultLabel(LOCTEXT("ItemLabelColumnLabel", "Item Label"))
+				.FillWidth(0.33f)
+
+				+SHeaderRow::Column(Column_ItemType)
+				.DefaultLabel(LOCTEXT("ItemTypeColumnLabel", "Item Type"))
+				.FillWidth(0.33f)
+
+				+SHeaderRow::Column(Column_Configuration)
+				.DefaultLabel(LOCTEXT("ConfigurationColumnLabel", "Configuration"))
+				.FillWidth(0.33f)			
+			)
+			.OnGetChildren(this, &SMediaProfileSourcesTreeView::GetTreeItemChildren)
+			.OnGenerateRow(this, &SMediaProfileSourcesTreeView::GenerateTreeItemRow)
+			.OnContextMenuOpening(this, &SMediaProfileSourcesTreeView::CreateContextMenu)
+		]
+	];
+
+	// Expand all elements in the tree view
+	for (const FMediaTreeItemPtr& TreeItem : MediaTreeItems)
+	{
+		MediaTreeView->SetItemExpansion(TreeItem, true);
+	}
+	
+	FCoreUObjectDelegates::OnObjectPropertyChanged.AddSP(this, &SMediaProfileSourcesTreeView::OnObjectPropertyChanged);
+	GEditor->RegisterForUndo(this);
+}
+
+SMediaProfileSourcesTreeView::~SMediaProfileSourcesTreeView()
+{
+	FCoreUObjectDelegates::OnObjectPropertyChanged.RemoveAll(this);
+	GEditor->UnregisterForUndo(this);
+}
+
+FReply SMediaProfileSourcesTreeView::OnKeyDown(const FGeometry& MyGeometry, const FKeyEvent& InKeyEvent)
+{
+	if (this->CommandList->ProcessCommandBindings(InKeyEvent))
+	{
+		return FReply::Handled();
+	}
+	
+	return FReply::Unhandled();
+}
+
+void SMediaProfileSourcesTreeView::PostUndo(bool bSuccess)
+{
+	// If the list of sources or outputs has changed in the media profile from the undo, update the media profile settings config
+	constexpr bool bSilent = true;
+	GetMutableDefault<UMediaProfileSettings>()->FillDefaultMediaSourceProxies(MediaProfile->NumMediaSources(), bSilent);
+	GetMutableDefault<UMediaProfileSettings>()->FillDefaultMediaOutputProxies(MediaProfile->NumMediaOutputs(), bSilent);
+}
+
+void SMediaProfileSourcesTreeView::PostRedo(bool bSuccess)
+{
+	PostUndo(bSuccess);
+}
+
+void SMediaProfileSourcesTreeView::BindCommands()
+{
+	CommandList->MapAction(FGenericCommands::Get().Copy,
+		FExecuteAction::CreateSP(this, &SMediaProfileSourcesTreeView::CopySelectedMedia),
+		FCanExecuteAction::CreateSP(this, &SMediaProfileSourcesTreeView::CanCopySelectedMedia));
+
+	CommandList->MapAction(FGenericCommands::Get().Cut,
+		FExecuteAction::CreateSP(this, &SMediaProfileSourcesTreeView::CutSelectedMedia),
+		FCanExecuteAction::CreateSP(this, &SMediaProfileSourcesTreeView::CanCutSelectedMedia));
+
+	CommandList->MapAction(FGenericCommands::Get().Paste,
+		FExecuteAction::CreateSP(this, &SMediaProfileSourcesTreeView::PasteSelectedMedia),
+		FCanExecuteAction::CreateSP(this, &SMediaProfileSourcesTreeView::CanPasteSelectedMedia));
+
+	CommandList->MapAction(FGenericCommands::Get().Duplicate,
+	FExecuteAction::CreateSP(this, &SMediaProfileSourcesTreeView::DuplicateSelectedMedia),
+	FCanExecuteAction::CreateSP(this, &SMediaProfileSourcesTreeView::CanDuplicateSelectedMedia));
+	
+	CommandList->MapAction(FGenericCommands::Get().Delete,
+		FExecuteAction::CreateSP(this, &SMediaProfileSourcesTreeView::DeleteSelectedMedia),
+		FCanExecuteAction::CreateSP(this, &SMediaProfileSourcesTreeView::CanDeleteSelectedMedia));
+}
+
+void SMediaProfileSourcesTreeView::CopySelectedMedia()
+{
+	TArray<UObject*> ObjectsToCopy;
+	for (const FMediaTreeItemPtr& SelectedItem : MediaTreeView->GetSelectedItems())
+	{
+		if (SelectedItem->BackingObject.IsValid())
+		{
+			ObjectsToCopy.Add(SelectedItem->BackingObject.Get());
+		}
+	}
+	
+	UnMarkAllObjects(EObjectMark(OBJECTMARK_TagExp | OBJECTMARK_TagImp));
+
+	FStringOutputDevice Archive;
+	const FExportObjectInnerContext Context;
+
+	for (UObject* Object : ObjectsToCopy)
+	{
+		UExporter::ExportToOutputDevice(&Context, Object, nullptr, Archive, TEXT("copy"), 0, PPF_ExportsNotFullyQualified | PPF_Copy | PPF_Delimited, false, nullptr);
+	}
+
+	FPlatformApplicationMisc::ClipboardCopy(*Archive);
+}
+
+bool SMediaProfileSourcesTreeView::CanCopySelectedMedia() const
+{
+	TArray<FMediaTreeItemPtr> SelectedItems = MediaTreeView->GetSelectedItems();
+	for (const FMediaTreeItemPtr& SelectedItem : SelectedItems)
+	{
+		// If any of the selected media sources are copyable (are non-null media sources or outputs), allow the copy action
+		if (SelectedItem->IsLeaf() && SelectedItem->BackingObject.IsValid())
+		{
+			return true;
+		}
+	}
+	
+	return false;
+}
+
+void SMediaProfileSourcesTreeView::CutSelectedMedia()
+{
+	CopySelectedMedia();
+	DeleteSelectedMedia();
+}
+
+bool SMediaProfileSourcesTreeView::CanCutSelectedMedia() const
+{
+	return CanCopySelectedMedia() && CanDeleteSelectedMedia();
+}
+
+/** Object factory to construct a list of media sources and outputs from the clipboard */
+class FMediaObjectTextFactory : public FCustomizableTextObjectFactory
+{
+public:
+	FMediaObjectTextFactory() : FCustomizableTextObjectFactory(GWarn) { }
+
+	virtual bool CanCreateClass(UClass* ObjectClass, bool& bOmitSubObjs) const override
+	{
+		const bool bIsMediaSource = ObjectClass->IsChildOf(UMediaSource::StaticClass());
+		const bool bIsMediaOutput = ObjectClass->IsChildOf(UMediaOutput::StaticClass());
+
+		return bIsMediaSource || bIsMediaOutput;
+	}
+
+	virtual void ProcessConstructedObject(UObject* NewObject) override
+	{
+		check(NewObject);
+
+		if (UMediaSource* MediaSource = Cast<UMediaSource>(NewObject))
+		{
+			MediaSources.Add(MediaSource);
+		}
+		else if (UMediaOutput* MediaOutput = Cast<UMediaOutput>(NewObject))
+		{
+			MediaOutputs.Add(MediaOutput);
+		}
+	}
+
+public:
+	TArray<UMediaSource*> MediaSources;
+	TArray<UMediaOutput*> MediaOutputs;
+};
+
+void SMediaProfileSourcesTreeView::PasteSelectedMedia()
+{
+	if (!MediaProfile.IsValid())
+	{
+		return;
+	}
+	
+	FScopedTransaction PasteMediaTransaction(LOCTEXT("PasteMediaTransaction", "Paste Media"));
+	MediaProfile->Modify();
+	
+	FString ClipboardContent;
+	FPlatformApplicationMisc::ClipboardPaste(ClipboardContent);
+
+	FMediaObjectTextFactory Factory;
+	Factory.ProcessBuffer(GetTransientPackage(), RF_Transactional, ClipboardContent);
+
+	int32 NumSourcesAdded = 0;
+	for (UMediaSource* MediaSourceToCopy : Factory.MediaSources)
+	{
+		UMediaSource* MediaSource = DuplicateObject(MediaSourceToCopy, MediaProfile.Get());
+		MediaProfile->AddMediaSource(MediaSource);
+		++NumSourcesAdded;
+	}
+
+	int32 NumOutputsAdded = 0;
+	for (UMediaOutput* MediaOutputToCopy : Factory.MediaOutputs)
+	{
+		UMediaOutput* MediaOutput = DuplicateObject(MediaOutputToCopy, MediaProfile.Get());
+		MediaProfile->AddMediaOutput(MediaOutput);
+		++NumOutputsAdded;
+	}
+
+	// Now, update the media profile project settings so that the proxy media slots match
+	constexpr bool bSilent = true;
+	if (NumSourcesAdded > 0)
+	{
+		GetMutableDefault<UMediaProfileSettings>()->FillDefaultMediaSourceProxies(MediaProfile->NumMediaSources(), bSilent);
+	}
+	
+	if (NumOutputsAdded > 0)
+	{
+		GetMutableDefault<UMediaProfileSettings>()->FillDefaultMediaOutputProxies(MediaProfile->NumMediaOutputs(), bSilent);
+	}
+	
+	MediaSourcesTreeView::TrySaveMediaProfileSettings();
+
+	FillMediaTree();
+}
+
+bool SMediaProfileSourcesTreeView::CanPasteSelectedMedia() const
+{
+	// Can't paste unless the clipboard has a string in it
+	FString ClipboardContent;
+	FPlatformApplicationMisc::ClipboardPaste(ClipboardContent);
+	if (!ClipboardContent.IsEmpty())
+	{
+		FMediaObjectTextFactory Factory;
+		Factory.ProcessBuffer(GetTransientPackage(), RF_Transactional, ClipboardContent);
+		return Factory.MediaSources.Num() > 0 || Factory.MediaOutputs.Num() > 0;
+	}
+
+	return false;
+}
+
+void SMediaProfileSourcesTreeView::DuplicateSelectedMedia()
+{
+	CopySelectedMedia();
+	PasteSelectedMedia();
+}
+
+bool SMediaProfileSourcesTreeView::CanDuplicateSelectedMedia() const
+{
+	return CanCopySelectedMedia();
+}
+
+void SMediaProfileSourcesTreeView::DeleteSelectedMedia()
+{
+	if (!MediaProfile.IsValid())
+	{
+		return;
+	}
+	
+	TArray<FMediaTreeItemPtr> SelectedItems = MediaTreeView->GetSelectedItems();
+	
+	TArray<int32> SourcesToRemove;
+	TArray<int32> OutputsToRemove;
+	for (const FMediaTreeItemPtr& SelectedItem : SelectedItems)
+	{
+		if (SelectedItem->BackingClass.IsValid())
+		{
+			if (SelectedItem->BackingClass.Get() == UMediaSource::StaticClass())
+			{
+				SourcesToRemove.Add(SelectedItem->Index);
+			}
+			else if (SelectedItem->BackingClass.Get() == UMediaOutput::StaticClass())
+			{
+				OutputsToRemove.Add(SelectedItem->Index);
+			}
+		}
+	}
+
+	if (SourcesToRemove.Num() == 0 && OutputsToRemove.Num() == 0)
+	{
+		return;
+	}
+
+	FScopedTransaction Transaction(LOCTEXT("DeleteMediaTransaction", "Delete Media"));
+	MediaProfile->Modify();
+
+	// First, remove the media from the profile, to ensure that the correct media is moved into the correct slots
+	// Reverse iterate through the indices for proper removal
+	SourcesToRemove.Sort();
+	for (int32 Index = SourcesToRemove.Num() - 1; Index >= 0; --Index)
+	{
+		MediaProfile->RemoveMediaSource(SourcesToRemove[Index]);
+	}
+
+	OutputsToRemove.Sort();
+	for (int32 Index = OutputsToRemove.Num() - 1; Index >= 0; --Index)
+	{
+		MediaProfile->RemoveMediaOutput(OutputsToRemove[Index]);
+	}
+
+	// Then, update the media profile project settings so that the proxy media slots match
+	constexpr bool bSilent = true;
+	GetMutableDefault<UMediaProfileSettings>()->FillDefaultMediaSourceProxies(MediaProfile->NumMediaSources(), bSilent);
+	GetMutableDefault<UMediaProfileSettings>()->FillDefaultMediaOutputProxies(MediaProfile->NumMediaOutputs(), bSilent);
+	
+	MediaSourcesTreeView::TrySaveMediaProfileSettings();
+
+	FillMediaTree();
+}
+
+bool SMediaProfileSourcesTreeView::CanDeleteSelectedMedia() const
+{
+	for (const FMediaTreeItemPtr& SelectedItem : MediaTreeView->GetSelectedItems())
+	{
+		if (SelectedItem->IsLeaf() && SelectedItem->BackingClass.IsValid())
+		{
+			return true;
+		}
+	}
+	
+	return false;
+}
+
+void SMediaProfileSourcesTreeView::AddNewMediaSource()
+{
+	if (MediaProfile.IsValid())
+	{
+		const int32 NumSources = MediaProfile->NumMediaSources();
+		if (NumSources >= 16)
+		{
+			FMessageDialog::Open(EAppMsgType::Ok, LOCTEXT("MaxMediaSourcesMessage", "Media profile only supports having up to 16 media sources"));
+			return;
+		}
+		
+		// TODO: Future proofing for when the media profile settings are reworked to fill from the active media profile sources/outputs instead of the other way around
+		FScopedTransaction AddMediaSourceTransaction(LOCTEXT("AddMediaSourceTransaction", "Add new Media Source"));
+		MediaProfile->Modify();
+		MediaProfile->AddMediaSource(nullptr);
+
+		GetMutableDefault<UMediaProfileSettings>()->FillDefaultMediaSourceProxies(NumSources + 1);
+
+		MediaSourcesTreeView::TrySaveMediaProfileSettings();
+
+		FillMediaTree();
+	}
+}
+
+void SMediaProfileSourcesTreeView::AddNewMediaOutput()
+{
+	if (MediaProfile.IsValid())
+	{
+		const int32 NumOutputs = MediaProfile->NumMediaOutputs();
+		if (NumOutputs >= 16)
+		{
+			FMessageDialog::Open(EAppMsgType::Ok, LOCTEXT("MaxMediaOutputsMessage", "Media profile only supports having up to 16 media outputs"));
+			return;
+		}
+		
+		// TODO: Future proofing for when the media profile settings are reworked to fill from the active media profile sources/outputs instead of the other way around
+		FScopedTransaction AddMediaOutputTransaction(LOCTEXT("AddMediaOutputTransaction", "Add new Media Output"));
+		MediaProfile->Modify();
+		MediaProfile->AddMediaOutput(nullptr);
+
+		GetMutableDefault<UMediaProfileSettings>()->FillDefaultMediaOutputProxies(NumOutputs + 1);
+
+		MediaSourcesTreeView::TrySaveMediaProfileSettings();
+
+		FillMediaTree();
+	}
+}
+
+void SMediaProfileSourcesTreeView::FillMediaTree()
+{
+	MediaTreeItems.Empty();
+
+	if (MediaProfile.IsValid())
+	{
+		FMediaTreeItemPtr MediaSourcesTreeItem = MakeShared<FMediaTreeItem>();
+		MediaSourcesTreeItem->Label = LOCTEXT("MediaSourcesTreeItemLabel", "Media Sources");
+		
+		for (int32 Index = 0; Index < MediaProfile->NumMediaSources(); ++Index)
+		{
+			UMediaSource* MediaSource = MediaProfile->GetMediaSource(Index);
+			FMediaTreeItemPtr MediaSourceTreeItem = MakeShared<FMediaTreeItem>();
+			MediaSourceTreeItem->BackingClass = UMediaSource::StaticClass();
+			MediaSourceTreeItem->BackingObject = MediaSource;
+			MediaSourceTreeItem->Index = Index;
+
+			// If the media source is a proxy media source, we want to output the information of the referenced media source asset instead of the proxy
+			bool bIsProxy = false;
+			if (UProxyMediaSource* ProxyMediaSource = Cast<UProxyMediaSource>(MediaSource))
+			{
+				MediaSource = ProxyMediaSource->GetLeafMediaSource();
+				bIsProxy = true;
+			}
+
+			if (MediaSource)
+			{
+				MediaSourceTreeItem->Label = FText::FromString(MediaSource->GetName());
+				MediaSourceTreeItem->Type = MediaSource->GetClass()->GetDisplayNameText();
+				MediaSourceTreeItem->Icon = FSlateIconFinder::FindIconForClass(MediaSource->GetClass()).GetIcon();
+				MediaSourceTreeItem->Configuration = FText::FromString(MediaSource->GetDescriptionString());
+				MediaSourceTreeItem->bFilteredOut = Toolbar.IsValid() ? !Toolbar->ItemPassesFilters(MediaSourcesTreeItem) : false;
+			}
+			else
+			{
+				MediaSourceTreeItem->Label = bIsProxy ?
+					LOCTEXT("UnconfiguredProxyMediaSourceLabel", "Proxy Media Source not set") :
+					LOCTEXT("NoMediaSourceSetLabel", "No Media Source Set");
+			}
+
+			MediaSourcesTreeItem->Children.Add(MediaSourceTreeItem);
+		}
+
+		MediaTreeItems.Add(MediaSourcesTreeItem);
+		
+		FMediaTreeItemPtr MediaOutputsTreeItem = MakeShared<FMediaTreeItem>();
+		MediaOutputsTreeItem->Label = LOCTEXT("MediaOutputsTreeItemLabel", "Media Outputs");
+		
+		for (int32 Index = 0; Index < MediaProfile->NumMediaOutputs(); ++Index)
+		{
+			UMediaOutput* MediaOutput = MediaProfile->GetMediaOutput(Index);
+			FMediaTreeItemPtr MediaOutputTreeItem = MakeShared<FMediaTreeItem>();
+			MediaOutputTreeItem->BackingClass = UMediaOutput::StaticClass();
+			MediaOutputTreeItem->BackingObject = MediaOutput;
+			MediaOutputTreeItem->Index = Index;
+
+			// If the media output is a proxy media output, we want to output the information of the referenced media output asset instead of the proxy
+			bool bIsProxy = false;
+			if (UProxyMediaOutput* ProxyMediaOutput = Cast<UProxyMediaOutput>(MediaOutput))
+			{
+				MediaOutput = ProxyMediaOutput->GetLeafMediaOutput();
+				bIsProxy = true;
+			}
+
+			if (MediaOutput)
+			{
+				MediaOutputTreeItem->Label = FText::FromString(MediaOutput->GetName());
+				MediaOutputTreeItem->Type = MediaOutput->GetClass()->GetDisplayNameText();
+				MediaOutputTreeItem->Icon = FSlateIconFinder::FindIconForClass(MediaOutput->GetClass()).GetIcon();
+				MediaOutputTreeItem->Configuration = FText::FromString(MediaOutput->GetDescriptionString());
+				MediaOutputTreeItem->bFilteredOut = Toolbar.IsValid() ? !Toolbar->ItemPassesFilters(MediaSourcesTreeItem) : false;
+			}
+			else
+			{
+				MediaOutputTreeItem->Label = bIsProxy ?
+					LOCTEXT("UnconfiguredProxyMediaOutputLabel", "Proxy Media Output not set") :
+					LOCTEXT("NoMediaOutputSetLabel", "No Media Output Set");
+			}
+
+			MediaOutputsTreeItem->Children.Add(MediaOutputTreeItem);
+		}
+
+		MediaTreeItems.Add(MediaOutputsTreeItem);
+	}
+
+	if (MediaTreeView.IsValid())
+	{
+		MediaTreeView->RequestTreeRefresh();
+		
+		for (const FMediaTreeItemPtr& TreeItem : MediaTreeItems)
+		{
+			MediaTreeView->SetItemExpansion(TreeItem, true);
+		}
+	}
+}
+
+void SMediaProfileSourcesTreeView::GetTreeItemChildren(FMediaTreeItemPtr InTreeItem, TArray<FMediaTreeItemPtr>& OutChildren) const
+{
+	if (InTreeItem.IsValid())
+	{
+		for (FMediaTreeItemPtr& Child : InTreeItem->Children)
+		{
+			if (Child.IsValid() && !Child->bFilteredOut)
+			{
+				OutChildren.Add(Child);
+			}
+		}
+	}
+}
+
+TSharedRef<ITableRow> SMediaProfileSourcesTreeView::GenerateTreeItemRow(FMediaTreeItemPtr InTreeItem, const TSharedRef<STableViewBase>& InOwnerTableView)
+{
+	return SNew(SMediaTreeItemRow, InOwnerTableView, InTreeItem)
+		.OnMoveMediaItem(this, &SMediaProfileSourcesTreeView::OnMoveMediaItem);
+}
+
+TSharedPtr<SWidget> SMediaProfileSourcesTreeView::CreateContextMenu()
+{
+	const int32 NumItems = MediaTreeView->GetNumItemsSelected();
+	if (NumItems >= 1)
+	{
+		constexpr bool bShouldCloseWindowAfterMenuSelection = true;
+		FMenuBuilder MenuBuilder(bShouldCloseWindowAfterMenuSelection, CommandList);
+
+		MenuBuilder.AddMenuEntry(FGenericCommands::Get().Copy);
+		MenuBuilder.AddMenuEntry(FGenericCommands::Get().Cut);
+		MenuBuilder.AddMenuEntry(FGenericCommands::Get().Paste);
+		MenuBuilder.AddMenuEntry(FGenericCommands::Get().Duplicate);
+		MenuBuilder.AddMenuEntry(FGenericCommands::Get().Delete);
+		
+		return MenuBuilder.MakeWidget();
+	}
+
+	return SNullWidget::NullWidget;
+}
+
+void SMediaProfileSourcesTreeView::OnObjectPropertyChanged(UObject* InObject, FPropertyChangedEvent& InPropertyChangedEvent)
+{
+	if (!InObject || InObject != MediaProfile)
+	{
+		return;
+	}
+	
+	FillMediaTree();
+}
+
+void SMediaProfileSourcesTreeView::OnFilterChanged()
+{
+	for (FMediaTreeItemPtr& TreeItem : MediaTreeItems)
+	{
+		for (FMediaTreeItemPtr& ChildItem : TreeItem->Children)
+		{
+			ChildItem->bFilteredOut = !Toolbar->ItemPassesFilters(ChildItem);
+		}
+	}
+
+	if (MediaTreeView.IsValid())
+	{
+		MediaTreeView->RequestTreeRefresh();
+		
+		for (const FMediaTreeItemPtr& TreeItem : MediaTreeItems)
+		{
+			MediaTreeView->SetItemExpansion(TreeItem, true);
+		}
+	}
+}
+
+void SMediaProfileSourcesTreeView::OnMoveMediaItem(TSharedPtr<FMediaTreeItem> InTreeItem, int32 InIndex)
+{
+	if (!MediaProfile.IsValid())
+	{
+		return;
+	}
+	
+	const bool bIsMediaSource = InTreeItem->BackingClass->IsChildOf<UMediaSource>();
+	const FText TransactionText = bIsMediaSource ? LOCTEXT("MoveMediaSourceTransaction", "Move Media Source") : LOCTEXT("MoveMediaOutputTransaction", "Move Media Output");
+	
+	FScopedTransaction MoveMediaItemTransaction(TransactionText);
+	MediaProfile->Modify();
+	
+	if (bIsMediaSource)
+	{
+		MediaProfile->MoveMediaSource(InTreeItem->Index, InIndex);
+	}
+	else
+	{
+		MediaProfile->MoveMediaOutput(InTreeItem->Index, InIndex);
+	}
+
+	FillMediaTree();
+}
+
+#undef LOCTEXT_NAMESPACE
diff --git a/Engine/Plugins/Media/MediaFrameworkUtilities/Source/MediaFrameworkUtilitiesEditor/Private/AssetEditor/SMediaProfileSourcesTreeView.h b/Engine/Plugins/Media/MediaFrameworkUtilities/Source/MediaFrameworkUtilitiesEditor/Private/AssetEditor/SMediaProfileSourcesTreeView.h
new file mode 100644
--- /dev/null
+++ b/Engine/Plugins/Media/MediaFrameworkUtilities/Source/MediaFrameworkUtilitiesEditor/Private/AssetEditor/SMediaProfileSourcesTreeView.h
@@ -0,0 +1,110 @@
+﻿// Copyright Epic Games, Inc. All Rights Reserved.
+
+#pragma once
+
+#include "CoreMinimal.h"
+#include "EditorUndoClient.h"
+#include "Misc/FilterCollection.h"
+#include "Widgets/SCompoundWidget.h"
+#include "Widgets/Views/STreeView.h"
+
+class SMediaProfileSourcesToolbar;
+class FUICommandList;
+class UMediaProfile;
+
+/**
+ * A tree view that displays a media profile's source and outputs
+ */
+class SMediaProfileSourcesTreeView : public SCompoundWidget, public FEditorUndoClient
+{
+private:
+	/** Struct to store information for each element displayed in the tree view */
+	struct FMediaTreeItem
+	{
+		int32 Index = INDEX_NONE;
+		FText Label = FText::GetEmpty();
+		FText Type = FText::GetEmpty();
+		FText Configuration = FText::GetEmpty();
+		const FSlateBrush* Icon = nullptr;
+		bool bFilteredOut = false;
+		
+		TArray<TSharedPtr<FMediaTreeItem>> Children;
+		TWeakObjectPtr<UClass> BackingClass = nullptr;
+		TWeakObjectPtr<UObject> BackingObject = nullptr;
+		
+		bool IsLeaf() const { return Index != INDEX_NONE; }
+	};
+	typedef TSharedPtr<FMediaTreeItem> FMediaTreeItemPtr;
+
+	// Names of the columns within the tree view
+	static const FLazyName Column_Index;
+	static const FLazyName Column_ItemLabel;
+	static const FLazyName Column_ItemType;
+	static const FLazyName Column_Configuration;
+	
+public:
+	SLATE_BEGIN_ARGS(SMediaProfileSourcesTreeView) {}
+	SLATE_END_ARGS()
+	
+	void Construct(const FArguments& InArgs, UMediaProfile* InMediaProfile);
+	virtual ~SMediaProfileSourcesTreeView() override;
+
+	// SWidget interface
+	virtual FReply OnKeyDown(const FGeometry& MyGeometry, const FKeyEvent& InKeyEvent) override;
+	// End of SWidget interface
+	
+	// FEditorUndoClient interface
+	virtual void PostUndo(bool bSuccess) override;
+	virtual void PostRedo(bool bSuccess) override;
+	// End of FEditorUndoClient
+	
+private:
+	/** Initializes and binds the command list to event handlers */
+	void BindCommands();
+
+	// Tree item context menu handlers
+	void CopySelectedMedia();
+	bool CanCopySelectedMedia() const;
+	void CutSelectedMedia();
+	bool CanCutSelectedMedia() const;
+	void PasteSelectedMedia();
+	bool CanPasteSelectedMedia() const;
+	void DuplicateSelectedMedia();
+	bool CanDuplicateSelectedMedia() const;
+	void DeleteSelectedMedia();
+	bool CanDeleteSelectedMedia() const;
+
+	/** Adds a new empty media source to the media profile */
+	void AddNewMediaSource();
+
+	/** Adds a new empty media output to the media profile */
+	void AddNewMediaOutput();
+
+	/** Recreates the tree view's source list from the current media profile's list of sources and outputs */
+	void FillMediaTree();
+
+	// TreeView event handlers
+	void GetTreeItemChildren(FMediaTreeItemPtr InTreeItem, TArray<FMediaTreeItemPtr>& OutChildren) const;
+	TSharedRef<ITableRow> GenerateTreeItemRow(FMediaTreeItemPtr InTreeItem, const TSharedRef<STableViewBase>& InOwnerTableView);
+	TSharedPtr<SWidget> CreateContextMenu();
+
+	/** Raised when an object property is changed. Used to monitor changes to the media profile and update the media tree view accordingly */
+	void OnObjectPropertyChanged(UObject* InObject, FPropertyChangedEvent& InPropertyChangedEvent);
+
+	/** Raised when the filter has changed in the toolbar, updates the tree view's source list to show only filtered items */
+	void OnFilterChanged();
+
+	/** Raised when a media tree item is moved */
+	void OnMoveMediaItem(TSharedPtr<FMediaTreeItem> InTreeItem, int32 InIndex);
+	
+private:
+	TArray<FMediaTreeItemPtr> MediaTreeItems;
+	TSharedPtr<STreeView<FMediaTreeItemPtr>> MediaTreeView;
+	TSharedPtr<FUICommandList> CommandList;
+	TWeakObjectPtr<UMediaProfile> MediaProfile;
+
+	TSharedPtr<SMediaProfileSourcesToolbar> Toolbar;
+	
+	friend class SMediaTreeItemRow;
+	friend class SMediaProfileSourcesToolbar;
+};
\ No newline at end of file
diff --git a/Engine/Plugins/Media/MediaFrameworkUtilities/Source/MediaFrameworkUtilitiesEditor/Private/AssetTypeActions/AssetTypeActions_MediaProfile.cpp b/Engine/Plugins/Media/MediaFrameworkUtilities/Source/MediaFrameworkUtilitiesEditor/Private/AssetTypeActions/AssetTypeActions_MediaProfile.cpp
--- a/Engine/Plugins/Media/MediaFrameworkUtilities/Source/MediaFrameworkUtilitiesEditor/Private/AssetTypeActions/AssetTypeActions_MediaProfile.cpp
+++ b/Engine/Plugins/Media/MediaFrameworkUtilities/Source/MediaFrameworkUtilitiesEditor/Private/AssetTypeActions/AssetTypeActions_MediaProfile.cpp
@@ -2,11 +2,19 @@
 
 #include "AssetTypeActions/AssetTypeActions_MediaProfile.h"
 
+#include "AssetEditor/MediaProfileEditor.h"
 #include "AssetEditor/MediaProfileEditorToolkit.h"
 #include "Profile/MediaProfile.h"
 
 #define LOCTEXT_NAMESPACE "MediaProfileEditor"
 
+/** Feature flag for the new media profile editor */
+static TAutoConsoleVariable<bool> CVarUseNewMediaProfileEditor(
+	TEXT("MediaProfile.UseNewMediaProfileEditor"), false,
+	TEXT("(Experimental) If true, the new media profile editor will be used when opening media profile assets.")
+	TEXT("If false (default), the existing media profile editor will be used."),
+	ECVF_Default);
+
 FText FAssetTypeActions_MediaProfile::GetName() const
 {
 	return LOCTEXT("AssetTypeActions_MediaProfile", "Media Profile");
@@ -23,7 +31,15 @@ void FAssetTypeActions_MediaProfile::OpenAssetEditor(const TArray<UObject*>& InO
 	{
 		if (UMediaProfile* Asset = Cast<UMediaProfile>(Obj))
 		{
-			FMediaProfileEditorToolkit::CreateEditor(EToolkitMode::Standalone, EditWithinLevelEditor, Asset);
+			const bool bUseNewMediaProfileEditor = CVarUseNewMediaProfileEditor.GetValueOnAnyThread();
+			if (bUseNewMediaProfileEditor)
+			{
+				FMediaProfileEditor::CreateMediaProfileEditor(EToolkitMode::Standalone, EditWithinLevelEditor, Asset);
+			}
+			else
+			{
+				FMediaProfileEditorToolkit::CreateEditor(EToolkitMode::Standalone, EditWithinLevelEditor, Asset);
+			}
 		}
 	}
 }
diff --git a/Engine/Plugins/Media/MediaIOFramework/Source/MediaIOCore/Private/FileMediaOutput.cpp b/Engine/Plugins/Media/MediaIOFramework/Source/MediaIOCore/Private/FileMediaOutput.cpp
--- a/Engine/Plugins/Media/MediaIOFramework/Source/MediaIOCore/Private/FileMediaOutput.cpp
+++ b/Engine/Plugins/Media/MediaIOFramework/Source/MediaIOCore/Private/FileMediaOutput.cpp
@@ -103,6 +103,12 @@ EMediaCaptureConversionOperation UFileMediaOutput::GetConversionOperation(EMedia
 	}
 }
 
+#if WITH_EDITOR
+FString UFileMediaOutput::GetDescriptionString() const
+{
+	return FilePath.Path;
+}
+#endif //WITH_EDITOR
 
 UMediaCapture* UFileMediaOutput::CreateMediaCaptureImpl()
 {
diff --git a/Engine/Plugins/Media/MediaIOFramework/Source/MediaIOCore/Public/FileMediaOutput.h b/Engine/Plugins/Media/MediaIOFramework/Source/MediaIOCore/Public/FileMediaOutput.h
--- a/Engine/Plugins/Media/MediaIOFramework/Source/MediaIOCore/Public/FileMediaOutput.h
+++ b/Engine/Plugins/Media/MediaIOFramework/Source/MediaIOCore/Public/FileMediaOutput.h
@@ -72,6 +72,10 @@ public:
 	virtual EPixelFormat GetRequestedPixelFormat() const override;
 	virtual EMediaCaptureConversionOperation GetConversionOperation(EMediaCaptureSourceType InSourceType) const override;
 
+#if WITH_EDITOR
+	virtual FString GetDescriptionString() const override;
+#endif //WITH_EDITOR
+	
 protected:
 	virtual UMediaCapture* CreateMediaCaptureImpl() override;
 	//~ End UMediaOutput interface
diff --git a/Engine/Plugins/Media/MediaIOFramework/Source/MediaIOCore/Public/MediaOutput.h b/Engine/Plugins/Media/MediaIOFramework/Source/MediaIOCore/Public/MediaOutput.h
--- a/Engine/Plugins/Media/MediaIOFramework/Source/MediaIOCore/Public/MediaOutput.h
+++ b/Engine/Plugins/Media/MediaIOFramework/Source/MediaIOCore/Public/MediaOutput.h
@@ -87,6 +87,9 @@ public:
 	virtual EMediaCaptureConversionOperation GetConversionOperation(EMediaCaptureSourceType InSourceType) const { return EMediaCaptureConversionOperation::NONE; }
 	
 #if WITH_EDITOR
+	/** Gets a string that sums up the media output's description for UI purposes */
+	virtual FString GetDescriptionString() const { return FString(); }
+	
 	DECLARE_MULTICAST_DELEGATE_OneParam(FOnOutputModified, UMediaOutput*);
 	FOnOutputModified& OnOutputModified() { return OnOutputModifiedDelegate; }
 	virtual void PostEditChangeProperty(struct FPropertyChangedEvent& PropertyChangedEvent) override;
diff --git a/Engine/Plugins/Media/NDIMedia/Source/NDIMedia/Private/NDIMediaOutput.cpp b/Engine/Plugins/Media/NDIMedia/Source/NDIMedia/Private/NDIMediaOutput.cpp
--- a/Engine/Plugins/Media/NDIMedia/Source/NDIMedia/Private/NDIMediaOutput.cpp
+++ b/Engine/Plugins/Media/NDIMedia/Source/NDIMedia/Private/NDIMediaOutput.cpp
@@ -57,6 +57,18 @@ EMediaCaptureConversionOperation UNDIMediaOutput::GetConversionOperation(EMediaC
 	}
 }
 
+#if WITH_EDITOR
+FString UNDIMediaOutput::GetDescriptionString() const
+{
+	const FIntPoint Size = GetRequestedSize();
+	return FString::Format(TEXT("{0}x{1} {2} fps {3}"), {
+		Size.X,
+		Size.Y,
+		FString::SanitizeFloat(FrameRate.AsDecimal()),
+		StaticEnum<ENDIMediaOutputPixelFormat>()->GetDisplayNameTextByValue(static_cast<int64>(DesiredPixelFormat)).ToString() });
+}
+#endif //WITH_EDITOR
+
 UMediaCapture* UNDIMediaOutput::CreateMediaCaptureImpl()
 {
 	UMediaCapture* Result = NewObject<UNDIMediaCapture>();
diff --git a/Engine/Plugins/Media/NDIMedia/Source/NDIMedia/Public/NDIMediaOutput.h b/Engine/Plugins/Media/NDIMedia/Source/NDIMedia/Public/NDIMediaOutput.h
--- a/Engine/Plugins/Media/NDIMedia/Source/NDIMedia/Public/NDIMediaOutput.h
+++ b/Engine/Plugins/Media/NDIMedia/Source/NDIMedia/Public/NDIMediaOutput.h
@@ -144,6 +144,10 @@ public:
 	virtual EPixelFormat GetRequestedPixelFormat() const override;
 	virtual EMediaCaptureConversionOperation GetConversionOperation(EMediaCaptureSourceType InSourceType) const override;
 
+#if WITH_EDITOR
+	virtual FString GetDescriptionString() const override;
+#endif //WITH_EDITOR
+
 protected:
 	virtual UMediaCapture* CreateMediaCaptureImpl() override;
 	//~ End UMediaOutput interface
diff --git a/Engine/Plugins/VirtualProduction/Rivermax/RivermaxMedia/Source/RivermaxMedia/Private/RivermaxMediaOutput.cpp b/Engine/Plugins/VirtualProduction/Rivermax/RivermaxMedia/Source/RivermaxMedia/Private/RivermaxMediaOutput.cpp
--- a/Engine/Plugins/VirtualProduction/Rivermax/RivermaxMedia/Source/RivermaxMedia/Private/RivermaxMediaOutput.cpp
+++ b/Engine/Plugins/VirtualProduction/Rivermax/RivermaxMedia/Source/RivermaxMedia/Private/RivermaxMediaOutput.cpp
@@ -57,6 +57,18 @@ EMediaCaptureConversionOperation URivermaxMediaOutput::GetConversionOperation(EM
 	return Result;
 }
 
+#if WITH_EDITOR
+FString URivermaxMediaOutput::GetDescriptionString() const
+{
+	const FIntPoint Size = GetRequestedSize();
+	return FString::Format(TEXT("{0}x{1} {2} fps {3}"), {
+		Size.X,
+		Size.Y,
+		FString::SanitizeFloat(VideoStream.FrameRate.AsDecimal()),
+		StaticEnum<ERivermaxMediaOutputPixelFormat>()->GetDisplayNameTextByValue(static_cast<int64>(VideoStream.PixelFormat)).ToString() });
+}
+#endif //WITH_EDITOR
+
 void URivermaxMediaOutput::ExportSDP(const FString& InPath)
 {
 	TArray<char> SDP;
diff --git a/Engine/Plugins/VirtualProduction/Rivermax/RivermaxMedia/Source/RivermaxMedia/Private/RivermaxMediaSource.cpp b/Engine/Plugins/VirtualProduction/Rivermax/RivermaxMedia/Source/RivermaxMedia/Private/RivermaxMediaSource.cpp
--- a/Engine/Plugins/VirtualProduction/Rivermax/RivermaxMedia/Source/RivermaxMedia/Private/RivermaxMediaSource.cpp
+++ b/Engine/Plugins/VirtualProduction/Rivermax/RivermaxMedia/Source/RivermaxMedia/Private/RivermaxMediaSource.cpp
@@ -164,3 +164,13 @@ bool URivermaxMediaSource::Validate() const
 	return true;
 }
 
+#if WITH_EDITOR
+FString URivermaxMediaSource::GetDescriptionString() const
+{
+	return FString::Format(TEXT("{0}x{1} {2} fps {3}"), {
+		Resolution.X,
+		Resolution.Y,
+		FString::SanitizeFloat(FrameRate.AsDecimal()),
+		StaticEnum<ERivermaxMediaSourcePixelFormat>()->GetDisplayNameTextByValue(static_cast<int64>(PixelFormat)).ToString() });
+}
+#endif //WITH_EDITOR
diff --git a/Engine/Plugins/VirtualProduction/Rivermax/RivermaxMedia/Source/RivermaxMedia/Public/RivermaxMediaOutput.h b/Engine/Plugins/VirtualProduction/Rivermax/RivermaxMedia/Source/RivermaxMedia/Public/RivermaxMediaOutput.h
--- a/Engine/Plugins/VirtualProduction/Rivermax/RivermaxMedia/Source/RivermaxMedia/Public/RivermaxMediaOutput.h
+++ b/Engine/Plugins/VirtualProduction/Rivermax/RivermaxMedia/Source/RivermaxMedia/Public/RivermaxMediaOutput.h
@@ -124,6 +124,10 @@ public:
 	virtual FIntPoint GetRequestedSize() const override;
 	virtual EPixelFormat GetRequestedPixelFormat() const override;
 	virtual EMediaCaptureConversionOperation GetConversionOperation(EMediaCaptureSourceType InSourceType) const override;
+	
+#if WITH_EDITOR
+	virtual FString GetDescriptionString() const override;
+#endif //WITH_EDITOR
 
 public:
 
diff --git a/Engine/Plugins/VirtualProduction/Rivermax/RivermaxMedia/Source/RivermaxMedia/Public/RivermaxMediaSource.h b/Engine/Plugins/VirtualProduction/Rivermax/RivermaxMedia/Source/RivermaxMedia/Public/RivermaxMediaSource.h
--- a/Engine/Plugins/VirtualProduction/Rivermax/RivermaxMedia/Source/RivermaxMedia/Public/RivermaxMediaSource.h
+++ b/Engine/Plugins/VirtualProduction/Rivermax/RivermaxMedia/Source/RivermaxMedia/Public/RivermaxMediaSource.h
@@ -127,5 +127,10 @@ public:
 	//~ Begin UMediaSource interface
 	virtual FString GetUrl() const override;
 	virtual bool Validate() const override;
+
+#if WITH_EDITOR
+	virtual FString GetDescriptionString() const override;
+#endif //WITH_EDITOR
+	
 	//~ End UMediaSource interface
 };
diff --git a/Engine/Source/Runtime/MediaAssets/Private/Assets/PlatformMediaSource.cpp b/Engine/Source/Runtime/MediaAssets/Private/Assets/PlatformMediaSource.cpp
--- a/Engine/Source/Runtime/MediaAssets/Private/Assets/PlatformMediaSource.cpp
+++ b/Engine/Source/Runtime/MediaAssets/Private/Assets/PlatformMediaSource.cpp
@@ -187,6 +187,22 @@ bool UPlatformMediaSource::Validate() const
 #endif
 }
 
+#if WITH_EDITOR
+FString UPlatformMediaSource::GetDescriptionString() const
+{
+	// Guard against reentrant calls.
+	static bool bIsGettingConfigText = false;
+	if (bIsGettingConfigText)
+	{
+		UE_LOG(LogMediaAssets, Warning, TEXT("UPlatformMediaSource::GetConfigurationDisplayText - Reentrant calls are not supported. Asset: %s"), *GetPathName());
+		return FString();
+	}
+	TGuardValue<bool> Guard(bIsGettingConfigText, true);
+
+	UMediaSource* PlatformMediaSource = GetMediaSource();
+	return (PlatformMediaSource != nullptr) ? PlatformMediaSource->GetDescriptionString() : FString();
+}
+#endif //WITH_EDITOR
 
 /* UPlatformMediaSource implementation
  *****************************************************************************/
diff --git a/Engine/Source/Runtime/MediaAssets/Public/MediaSource.h b/Engine/Source/Runtime/MediaAssets/Public/MediaSource.h
--- a/Engine/Source/Runtime/MediaAssets/Public/MediaSource.h
+++ b/Engine/Source/Runtime/MediaAssets/Public/MediaSource.h
@@ -81,6 +81,9 @@ public:
 	 */
 	void SetThumbnail(UTexture* InTexture) { ThumbnailImage = InTexture; }
 
+	/** Gets a string that sums up the media source's description for UI purposes */
+	MEDIAASSETS_API virtual FString GetDescriptionString() const { return GetUrl(); }
+	
 #endif // WITH_EDITOR
 
 	/**
diff --git a/Engine/Source/Runtime/MediaAssets/Public/PlatformMediaSource.h b/Engine/Source/Runtime/MediaAssets/Public/PlatformMediaSource.h
--- a/Engine/Source/Runtime/MediaAssets/Public/PlatformMediaSource.h
+++ b/Engine/Source/Runtime/MediaAssets/Public/PlatformMediaSource.h
@@ -55,6 +55,10 @@ public:
 	MEDIAASSETS_API virtual FString GetUrl() const override;
 	MEDIAASSETS_API virtual bool Validate() const override;
 
+#if WITH_EDITOR
+	virtual FString GetDescriptionString() const override;
+#endif //WITH_EDITOR
+	
 public:
 
 	//~ IMediaOptions interface
