Media Profile: Added custom details panel to new media profile editor that displays the properties for the selected media sources and outputs from the media outliner. Adds ability to label media sources and outputs in a media profile, as well as adding support for managing media capture for media profile outputs directly from the media profile editor.

diff --git a/Engine/Plugins/Media/AjaMedia/Source/AjaMedia/Private/Assets/AjaMediaSource.cpp b/Engine/Plugins/Media/AjaMedia/Source/AjaMedia/Private/Assets/AjaMediaSource.cpp
--- a/Engine/Plugins/Media/AjaMedia/Source/AjaMedia/Private/Assets/AjaMediaSource.cpp
+++ b/Engine/Plugins/Media/AjaMedia/Source/AjaMedia/Private/Assets/AjaMediaSource.cpp
@@ -293,6 +293,41 @@ FString UAjaMediaSource::GetDescriptionString() const
 
 	return Super::GetDescriptionString();
 }
+
+void UAjaMediaSource::GetDetailsPanelInfoElements(TArray<TTuple<FText, FText, FText>>& OutInfoElements) const
+{
+#define LOCTEXT_NAMESPACE "AjaMediaSource"
+	IMediaIOCoreDeviceProvider* DeviceProviderPtr = IMediaIOCoreModule::Get().GetDeviceProvider(FAjaDeviceProvider::GetProviderName());
+	if (!DeviceProviderPtr)
+	{
+		return;
+	}
+	
+	const FText ConfigHeader = LOCTEXT("MediaConfigurationHeader", "Media Configuration");
+
+	const FText SourceText = FText::Format(LOCTEXT("TransportFormat", "{0} {1}"),
+		DeviceProviderPtr->GetTransportName(MediaConfiguration.MediaConnection.TransportType, MediaConfiguration.MediaConnection.QuadTransportType),
+		FText::AsNumber(MediaConfiguration.MediaConnection.PortIdentifier));
+	
+	OutInfoElements.Add(FInfoElement(ConfigHeader, LOCTEXT("DeviceLabel", "Device"), FText::FromName(MediaConfiguration.MediaConnection.Device.DeviceName)));
+	OutInfoElements.Add(FInfoElement(ConfigHeader, LOCTEXT("SourceLabel", "Source"), SourceText));
+	OutInfoElements.Add(FInfoElement(ConfigHeader, LOCTEXT("ResolutionLabel", "Resolution"), MediaConfiguration.MediaMode.GetModeName()));
+	OutInfoElements.Add(FInfoElement(ConfigHeader, LOCTEXT("StandardLabel", "Standard"), UEnum::GetDisplayValueAsText(MediaConfiguration.MediaMode.Standard)));
+	OutInfoElements.Add(FInfoElement(ConfigHeader, LOCTEXT("FramerateLabel", "Framerate"), MediaConfiguration.MediaMode.FrameRate.ToPrettyText()));
+
+	const FText AdditionalHeader = LOCTEXT("AdditionalSettingsHeader", "Additional Settings");
+	const FText DisabledText = LOCTEXT("Disabled", "Disabled");
+	const FText AncillaryBufferText = FText::Format(LOCTEXT("AncillaryBufferFormat", "{0} buffered frames"), FText::AsNumber(MaxNumAncillaryFrameBuffer));
+	const FText AudioText = FText::Format(LOCTEXT("AudioFormat", "{0}, {1} buffered frames"),
+		UEnum::GetDisplayValueAsText(AudioChannel), FText::AsNumber(MaxNumAudioFrameBuffer));
+	
+	OutInfoElements.Add(FInfoElement(AdditionalHeader, LOCTEXT("TimecodeFormatLabel", "Timecode Format"), UEnum::GetDisplayValueAsText(AutoDetectableTimecodeFormat)));
+	OutInfoElements.Add(FInfoElement(AdditionalHeader, LOCTEXT("AncillaryLabel", "Ancillary"), bCaptureAncillary ? AncillaryBufferText : DisabledText));
+	OutInfoElements.Add(FInfoElement(AdditionalHeader, LOCTEXT("AudioLabel", "Audio"), bCaptureAudio ? AudioText : DisabledText));
+	OutInfoElements.Add(FInfoElement(AdditionalHeader, LOCTEXT("ColorConversionLabel", "Color Conversion"), FText::FromString(ColorConversionSettings.ToString())));
+	
+#undef LOCTEXT_NAMESPACE
+}
 #endif //WITH_EDITOR
 
 #if WITH_EDITOR
diff --git a/Engine/Plugins/Media/AjaMedia/Source/AjaMedia/Public/AjaMediaSource.h b/Engine/Plugins/Media/AjaMedia/Source/AjaMedia/Public/AjaMediaSource.h
--- a/Engine/Plugins/Media/AjaMedia/Source/AjaMedia/Public/AjaMediaSource.h
+++ b/Engine/Plugins/Media/AjaMedia/Source/AjaMedia/Public/AjaMediaSource.h
@@ -140,6 +140,7 @@ public:
 	
 #if WITH_EDITOR
 	virtual FString GetDescriptionString() const override;
+	virtual void GetDetailsPanelInfoElements(TArray<TTuple<FText, FText, FText>>& OutInfoElements) const override;
 #endif //WITH_EDITOR
 
 
diff --git a/Engine/Plugins/Media/AjaMedia/Source/AjaMediaOutput/Private/AjaMediaOutput.cpp b/Engine/Plugins/Media/AjaMedia/Source/AjaMediaOutput/Private/AjaMediaOutput.cpp
--- a/Engine/Plugins/Media/AjaMedia/Source/AjaMediaOutput/Private/AjaMediaOutput.cpp
+++ b/Engine/Plugins/Media/AjaMedia/Source/AjaMediaOutput/Private/AjaMediaOutput.cpp
@@ -216,6 +216,27 @@ FString UAjaMediaOutput::GetDescriptionString() const
 
 	return Super::GetDescriptionString();
 }
+
+void UAjaMediaOutput::GetDetailsPanelInfoElements(TArray<FInfoElement>& OutInfoElements) const
+{
+	IMediaIOCoreDeviceProvider* DeviceProviderPtr = IMediaIOCoreModule::Get().GetDeviceProvider(FAjaDeviceProvider::GetProviderName());
+	if (!DeviceProviderPtr)
+	{
+		return;
+	}
+	
+	const FText ConfigHeader = LOCTEXT("MediaConfigurationHeader", "Media Configuration");
+
+	const FText SourceText = FText::Format(LOCTEXT("TransportFormat", "{0} {1}"),
+		DeviceProviderPtr->GetTransportName(OutputConfiguration.MediaConfiguration.MediaConnection.TransportType, OutputConfiguration.MediaConfiguration.MediaConnection.QuadTransportType),
+		FText::AsNumber(OutputConfiguration.MediaConfiguration.MediaConnection.PortIdentifier));
+	
+	OutInfoElements.Add(FInfoElement(ConfigHeader, LOCTEXT("DeviceLabel", "Device"), FText::FromName(OutputConfiguration.MediaConfiguration.MediaConnection.Device.DeviceName)));
+	OutInfoElements.Add(FInfoElement(ConfigHeader, LOCTEXT("SourceLabel", "Source"), SourceText));
+	OutInfoElements.Add(FInfoElement(ConfigHeader, LOCTEXT("ResolutionLabel", "Resolution"), OutputConfiguration.MediaConfiguration.MediaMode.GetModeName()));
+	OutInfoElements.Add(FInfoElement(ConfigHeader, LOCTEXT("StandardLabel", "Standard"), UEnum::GetDisplayValueAsText(OutputConfiguration.MediaConfiguration.MediaMode.Standard)));
+	OutInfoElements.Add(FInfoElement(ConfigHeader, LOCTEXT("FramerateLabel", "Framerate"), OutputConfiguration.MediaConfiguration.MediaMode.FrameRate.ToPrettyText()));
+}
 #endif // WITH_EDITOR
 
 UMediaCapture* UAjaMediaOutput::CreateMediaCaptureImpl()
diff --git a/Engine/Plugins/Media/AjaMedia/Source/AjaMediaOutput/Public/AjaMediaOutput.h b/Engine/Plugins/Media/AjaMedia/Source/AjaMediaOutput/Public/AjaMediaOutput.h
--- a/Engine/Plugins/Media/AjaMedia/Source/AjaMediaOutput/Public/AjaMediaOutput.h
+++ b/Engine/Plugins/Media/AjaMedia/Source/AjaMediaOutput/Public/AjaMediaOutput.h
@@ -147,6 +147,7 @@ public:
 
 #if WITH_EDITOR
 	virtual FString GetDescriptionString() const override;
+	virtual void GetDetailsPanelInfoElements(TArray<FInfoElement>& OutInfoElements) const override;
 #endif //WITH_EDITOR
 	
 protected:
diff --git a/Engine/Plugins/Media/BlackmagicMedia/Source/BlackmagicMedia/Private/Assets/BlackmagicMediaSource.cpp b/Engine/Plugins/Media/BlackmagicMedia/Source/BlackmagicMedia/Private/Assets/BlackmagicMediaSource.cpp
--- a/Engine/Plugins/Media/BlackmagicMedia/Source/BlackmagicMedia/Private/Assets/BlackmagicMediaSource.cpp
+++ b/Engine/Plugins/Media/BlackmagicMedia/Source/BlackmagicMedia/Private/Assets/BlackmagicMediaSource.cpp
@@ -185,6 +185,39 @@ FString UBlackmagicMediaSource::GetDescriptionString() const
 
 	return Super::GetDescriptionString();
 }
+
+void UBlackmagicMediaSource::GetDetailsPanelInfoElements(TArray<FInfoElement>& OutInfoElements) const
+{
+#define LOCTEXT_NAMESPACE "BlackmagicMediaSource"
+	IMediaIOCoreDeviceProvider* DeviceProviderPtr = IMediaIOCoreModule::Get().GetDeviceProvider(FBlackmagicDeviceProvider::GetProviderName());
+	if (!DeviceProviderPtr)
+	{
+		return;
+	}
+	
+	const FText ConfigHeader = LOCTEXT("MediaConfigurationHeader", "Media Configuration");
+
+	const FText SourceText = FText::Format(LOCTEXT("TransportFormat", "{0} {1}"),
+		DeviceProviderPtr->GetTransportName(MediaConfiguration.MediaConnection.TransportType, MediaConfiguration.MediaConnection.QuadTransportType),
+		FText::AsNumber(MediaConfiguration.MediaConnection.PortIdentifier));
+	
+	OutInfoElements.Add(FInfoElement(ConfigHeader, LOCTEXT("DeviceLabel", "Device"), FText::FromName(MediaConfiguration.MediaConnection.Device.DeviceName)));
+	OutInfoElements.Add(FInfoElement(ConfigHeader, LOCTEXT("SourceLabel", "Source"), SourceText));
+	OutInfoElements.Add(FInfoElement(ConfigHeader, LOCTEXT("ResolutionLabel", "Resolution"), MediaConfiguration.MediaMode.GetModeName()));
+	OutInfoElements.Add(FInfoElement(ConfigHeader, LOCTEXT("StandardLabel", "Standard"), UEnum::GetDisplayValueAsText(MediaConfiguration.MediaMode.Standard)));
+	OutInfoElements.Add(FInfoElement(ConfigHeader, LOCTEXT("FramerateLabel", "Framerate"), MediaConfiguration.MediaMode.FrameRate.ToPrettyText()));
+
+	const FText AdditionalHeader = LOCTEXT("AdditionalSettingsHeader", "Additional Settings");
+	const FText DisabledText = LOCTEXT("Disabled", "Disabled");
+	const FText AudioText = FText::Format(LOCTEXT("AudioFormat", "{0}, {1} buffered frames"),
+		UEnum::GetDisplayValueAsText(AudioChannels), FText::AsNumber(MaxNumAudioFrameBuffer));
+	
+	OutInfoElements.Add(FInfoElement(AdditionalHeader, LOCTEXT("TimecodeFormatLabel", "Timecode Format"), UEnum::GetDisplayValueAsText(AutoDetectableTimecodeFormat)));
+	OutInfoElements.Add(FInfoElement(AdditionalHeader, LOCTEXT("AudioLabel", "Audio"), bCaptureAudio ? AudioText : DisabledText));
+	OutInfoElements.Add(FInfoElement(AdditionalHeader, LOCTEXT("ColorConversionLabel", "Color Conversion"), FText::FromString(ColorConversionSettings.ToString())));
+	
+#undef LOCTEXT_NAMESPACE
+}
 #endif //WITH_EDITOR
 
 void UBlackmagicMediaSource::AssignDefaultConfiguration()
diff --git a/Engine/Plugins/Media/BlackmagicMedia/Source/BlackmagicMedia/Public/BlackmagicMediaSource.h b/Engine/Plugins/Media/BlackmagicMedia/Source/BlackmagicMedia/Public/BlackmagicMediaSource.h
--- a/Engine/Plugins/Media/BlackmagicMedia/Source/BlackmagicMedia/Public/BlackmagicMediaSource.h
+++ b/Engine/Plugins/Media/BlackmagicMedia/Source/BlackmagicMedia/Public/BlackmagicMediaSource.h
@@ -126,6 +126,7 @@ public:
 
 #if WITH_EDITOR
 	virtual FString GetDescriptionString() const override;
+	virtual void GetDetailsPanelInfoElements(TArray<FInfoElement>& OutInfoElements) const override;
 #endif //WITH_EDITOR
 	
 public:
diff --git a/Engine/Plugins/Media/BlackmagicMedia/Source/BlackmagicMediaOutput/Private/BlackmagicMediaOutput.cpp b/Engine/Plugins/Media/BlackmagicMedia/Source/BlackmagicMediaOutput/Private/BlackmagicMediaOutput.cpp
--- a/Engine/Plugins/Media/BlackmagicMedia/Source/BlackmagicMediaOutput/Private/BlackmagicMediaOutput.cpp
+++ b/Engine/Plugins/Media/BlackmagicMedia/Source/BlackmagicMediaOutput/Private/BlackmagicMediaOutput.cpp
@@ -149,6 +149,35 @@ FString UBlackmagicMediaOutput::GetDescriptionString() const
 
 	return Super::GetDescriptionString();
 }
+
+void UBlackmagicMediaOutput::GetDetailsPanelInfoElements(TArray<FInfoElement>& OutInfoElements) const
+{
+	IMediaIOCoreDeviceProvider* DeviceProviderPtr = IMediaIOCoreModule::Get().GetDeviceProvider(FBlackmagicDeviceProvider::GetProviderName());
+	if (!DeviceProviderPtr)
+	{
+		return;
+	}
+	
+	const FText ConfigHeader = LOCTEXT("MediaConfigurationHeader", "Media Configuration");
+
+	const FText SourceText = FText::Format(LOCTEXT("TransportFormat", "{0} {1}"),
+		DeviceProviderPtr->GetTransportName(OutputConfiguration.MediaConfiguration.MediaConnection.TransportType, OutputConfiguration.MediaConfiguration.MediaConnection.QuadTransportType),
+		FText::AsNumber(OutputConfiguration.MediaConfiguration.MediaConnection.PortIdentifier));
+	
+	OutInfoElements.Add(FInfoElement(ConfigHeader, LOCTEXT("DeviceLabel", "Device"), FText::FromName(OutputConfiguration.MediaConfiguration.MediaConnection.Device.DeviceName)));
+	OutInfoElements.Add(FInfoElement(ConfigHeader, LOCTEXT("SourceLabel", "Source"), SourceText));
+	OutInfoElements.Add(FInfoElement(ConfigHeader, LOCTEXT("ResolutionLabel", "Resolution"), OutputConfiguration.MediaConfiguration.MediaMode.GetModeName()));
+	OutInfoElements.Add(FInfoElement(ConfigHeader, LOCTEXT("StandardLabel", "Standard"), UEnum::GetDisplayValueAsText(OutputConfiguration.MediaConfiguration.MediaMode.Standard)));
+	OutInfoElements.Add(FInfoElement(ConfigHeader, LOCTEXT("FramerateLabel", "Framerate"), OutputConfiguration.MediaConfiguration.MediaMode.FrameRate.ToPrettyText()));
+
+	const FText AdditionalHeader = LOCTEXT("AdditionalSettingsHeader", "Additional Settings");
+	const FText DisabledText = LOCTEXT("Disabled", "Disabled");
+	const FText AudioText = FText::Format(LOCTEXT("AudioFormat", "{0} channels, {1} sample rate"),
+		UEnum::GetDisplayValueAsText(OutputChannelCount), UEnum::GetDisplayValueAsText(AudioSampleRate));
+	
+	OutInfoElements.Add(FInfoElement(AdditionalHeader, LOCTEXT("TimecodeFormatLabel", "Timecode Format"), UEnum::GetDisplayValueAsText(TimecodeFormat)));
+	OutInfoElements.Add(FInfoElement(AdditionalHeader, LOCTEXT("AudioLabel", "Audio"), bOutputAudio ? AudioText : DisabledText));
+}
 #endif //WITH_EDITOR
 
 UMediaCapture* UBlackmagicMediaOutput::CreateMediaCaptureImpl()
diff --git a/Engine/Plugins/Media/BlackmagicMedia/Source/BlackmagicMediaOutput/Public/BlackmagicMediaOutput.h b/Engine/Plugins/Media/BlackmagicMedia/Source/BlackmagicMediaOutput/Public/BlackmagicMediaOutput.h
--- a/Engine/Plugins/Media/BlackmagicMedia/Source/BlackmagicMediaOutput/Public/BlackmagicMediaOutput.h
+++ b/Engine/Plugins/Media/BlackmagicMedia/Source/BlackmagicMediaOutput/Public/BlackmagicMediaOutput.h
@@ -153,6 +153,7 @@ public:
 
 #if WITH_EDITOR
 	virtual FString GetDescriptionString() const override;
+	virtual void GetDetailsPanelInfoElements(TArray<FInfoElement>& OutInfoElements) const override;
 #endif //WITH_EDITOR
 	
 protected:
diff --git a/Engine/Plugins/Media/ImgMedia/Source/ImgMedia/Private/Assets/ImgMediaSource.cpp b/Engine/Plugins/Media/ImgMedia/Source/ImgMedia/Private/Assets/ImgMediaSource.cpp
--- a/Engine/Plugins/Media/ImgMedia/Source/ImgMedia/Private/Assets/ImgMediaSource.cpp
+++ b/Engine/Plugins/Media/ImgMedia/Source/ImgMedia/Private/Assets/ImgMediaSource.cpp
@@ -278,6 +278,13 @@ bool UImgMediaSource::Validate() const
 }
 
 #if WITH_EDITOR
+void UImgMediaSource::GetDetailsPanelInfoElements(TArray<FInfoElement>& OutInfoElements) const
+{
+	OutInfoElements.Add(FInfoElement(
+		NSLOCTEXT("ImgMediaSource", "MediaConfigurationHeader", "Media Configuration"),
+		NSLOCTEXT("ImgMediaSource", "FilePathLabel", "File Path"),
+		FText::FromString(GetFullPath())));
+}
 
 void UImgMediaSource::PostEditChangeProperty(FPropertyChangedEvent& PropertyChangedEvent)
 {
diff --git a/Engine/Plugins/Media/ImgMedia/Source/ImgMedia/Public/ImgMediaSource.h b/Engine/Plugins/Media/ImgMedia/Source/ImgMedia/Public/ImgMediaSource.h
--- a/Engine/Plugins/Media/ImgMedia/Source/ImgMedia/Public/ImgMediaSource.h
+++ b/Engine/Plugins/Media/ImgMedia/Source/ImgMedia/Public/ImgMediaSource.h
@@ -167,6 +167,8 @@ public:
 	UE_API virtual bool Validate() const override;
 
 #if WITH_EDITOR
+	virtual void GetDetailsPanelInfoElements(TArray<FInfoElement>& OutInfoElements) const override;
+	
 	UE_API virtual void PostEditChangeProperty(FPropertyChangedEvent& PropertyChangedEvent) override;
 #endif // WITH_EDITOR
 
diff --git a/Engine/Plugins/Media/MediaFrameworkUtilities/Source/MediaFrameworkUtilities/Private/MediaAssets/ProxyMediaOutput.cpp b/Engine/Plugins/Media/MediaFrameworkUtilities/Source/MediaFrameworkUtilities/Private/MediaAssets/ProxyMediaOutput.cpp
--- a/Engine/Plugins/Media/MediaFrameworkUtilities/Source/MediaFrameworkUtilities/Private/MediaAssets/ProxyMediaOutput.cpp
+++ b/Engine/Plugins/Media/MediaFrameworkUtilities/Source/MediaFrameworkUtilities/Private/MediaAssets/ProxyMediaOutput.cpp
@@ -101,6 +101,23 @@ FString UProxyMediaOutput::GetDescriptionString() const
 	UMediaOutput* CurrentProxy = GetMediaOutput();
 	return (CurrentProxy != nullptr) ? CurrentProxy->GetDescriptionString() : FString();
 }
+
+void UProxyMediaOutput::GetDetailsPanelInfoElements(TArray<FInfoElement>& OutInfoElements) const
+{
+	// Guard against reentrant calls.
+	if (bGetInfoElementsGuard)
+	{
+		UE_LOG(LogMediaFrameworkUtilities, Warning, TEXT("UProxyMediaOutput::GetDetailsPanelInfoElements - Reentrant calls are not supported. Asset: %s"), *GetPathName());
+		return;
+	}
+	TGuardValue<bool> Guard(bGetInfoElementsGuard, true);
+
+	UMediaOutput* CurrentProxy = GetMediaOutput();
+	if (CurrentProxy != nullptr)
+	{
+		CurrentProxy->GetDetailsPanelInfoElements(OutInfoElements);
+	}
+}
 #endif //WITH_EDITOR
 
 UMediaCapture* UProxyMediaOutput::CreateMediaCaptureImpl()
diff --git a/Engine/Plugins/Media/MediaFrameworkUtilities/Source/MediaFrameworkUtilities/Private/MediaAssets/ProxyMediaSource.cpp b/Engine/Plugins/Media/MediaFrameworkUtilities/Source/MediaFrameworkUtilities/Private/MediaAssets/ProxyMediaSource.cpp
--- a/Engine/Plugins/Media/MediaFrameworkUtilities/Source/MediaFrameworkUtilities/Private/MediaAssets/ProxyMediaSource.cpp
+++ b/Engine/Plugins/Media/MediaFrameworkUtilities/Source/MediaFrameworkUtilities/Private/MediaAssets/ProxyMediaSource.cpp
@@ -63,6 +63,23 @@ FString UProxyMediaSource::GetDescriptionString() const
 	UMediaSource* CurrentProxy = GetMediaSource();
 	return (CurrentProxy != nullptr) ? CurrentProxy->GetDescriptionString() : FString();
 }
+
+void UProxyMediaSource::GetDetailsPanelInfoElements(TArray<FInfoElement>& OutInfoElements) const
+{
+	// Guard against reentrant calls.
+	if (bGetInfoElementsGuard)
+	{
+		UE_LOG(LogMediaFrameworkUtilities, Warning, TEXT("UMediaSourceProxy::GetDetailsPanelInfoElements - Reentrant calls are not supported. Asset: %s"), *GetPathName());
+		return;
+	}
+	TGuardValue<bool> Guard(bGetInfoElementsGuard, true);
+
+	UMediaSource* CurrentProxy = GetMediaSource();
+	if (CurrentProxy != nullptr)
+	{
+		CurrentProxy->GetDetailsPanelInfoElements(OutInfoElements);
+	}
+}
 #endif //WITH_EDITOR
 
 /* UMediaSourceProxy implementation
diff --git a/Engine/Plugins/Media/MediaFrameworkUtilities/Source/MediaFrameworkUtilities/Private/Profile/MediaProfile.cpp b/Engine/Plugins/Media/MediaFrameworkUtilities/Source/MediaFrameworkUtilities/Private/Profile/MediaProfile.cpp
--- a/Engine/Plugins/Media/MediaFrameworkUtilities/Source/MediaFrameworkUtilities/Private/Profile/MediaProfile.cpp
+++ b/Engine/Plugins/Media/MediaFrameworkUtilities/Source/MediaFrameworkUtilities/Private/Profile/MediaProfile.cpp
@@ -52,15 +52,56 @@ UMediaSource* UMediaProfile::GetMediaSource(int32 Index) const
 	return nullptr;
 }
 
+void UMediaProfile::SetMediaSource(int32 Index, UMediaSource* InMediaSource)
+{
+	if (MediaSources.IsValidIndex(Index))
+	{
+		FString CurrentMediaSourceName = MediaSources[Index] ? MediaSources[Index]->GetName() : TEXT("");
+		MediaSources[Index] = InMediaSource;
+		
+#if WITH_EDITOR
+		FPropertyChangedEvent PropertyChangedEvent(StaticClass()->FindPropertyByName(GET_MEMBER_NAME_CHECKED(UMediaProfile, MediaSources)), EPropertyChangeType::ValueSet);
+		FCoreUObjectDelegates::OnObjectPropertyChanged.Broadcast(this, PropertyChangedEvent);
+		
+		// If the original label for the media source was the source's object name (indicating the user has not changed it)
+		// update the label to the new source's object name
+		if (!MediaSourceLabels.IsValidIndex(Index) || GetLabelForMediaSource(Index) == CurrentMediaSourceName)
+		{
+			SetLabelForMediaSource(Index, InMediaSource->GetName());
+		}
+#endif
+	}
+}
 
 int32 UMediaProfile::NumMediaSources() const
 {
 	return MediaSources.Num();
 }
 
+int32 UMediaProfile::FindMediaSourceIndex(UMediaSource* MediaSource) const
+{
+	return MediaSources.Find(MediaSource);
+}
+
 void UMediaProfile::AddMediaSource(UMediaSource* MediaSource)
 {
 	MediaSources.Add(MediaSource);
+
+#if WITH_EDITOR
+	if (MediaSourceLabels.Num() < MediaSources.Num())
+	{
+		for (int32 Index = MediaSourceLabels.Num(); Index < MediaSources.Num(); ++Index)
+		{
+			FString Name = MediaSources[Index] ? MediaSources[Index]->GetName() : TEXT("");
+			MediaSourceLabels.Add(Name);
+		}
+	}
+	else
+	{
+		FString Name = MediaSource ? MediaSource->GetName() : TEXT("");
+		MediaSourceLabels.Add(Name);
+	}
+#endif
 }
 
 bool UMediaProfile::RemoveMediaSource(int32 Index)
@@ -71,6 +112,14 @@ bool UMediaProfile::RemoveMediaSource(int32 Index)
 	}
 	
 	MediaSources.RemoveAt(Index, EAllowShrinking::No);
+
+#if WITH_EDITOR
+	if (MediaSourceLabels.IsValidIndex(Index))
+	{
+		MediaSourceLabels.RemoveAt(Index, EAllowShrinking::No);
+	}
+#endif
+	
 	return true;
 }
 
@@ -84,10 +133,73 @@ bool UMediaProfile::MoveMediaSource(int32 CurrentIndex, int32 DestIndex)
 	TObjectPtr<UMediaSource> MediaSource = MediaSources[CurrentIndex];
 	MediaSources.RemoveAt(CurrentIndex, EAllowShrinking::No);
 	MediaSources.Insert(MediaSource, DestIndex);
+
+#if WITH_EDITOR
+	if (MediaSourceLabels.IsValidIndex(CurrentIndex))
+	{
+		FString Label = MediaSourceLabels[CurrentIndex];
+		MediaSourceLabels.RemoveAt(CurrentIndex, EAllowShrinking::No);
+		MediaSourceLabels.Insert(Label, DestIndex);
+	}
+#endif
 	
 	return true;
 }
 
+#if WITH_EDITOR
+FString UMediaProfile::GetLabelForMediaSource(int32 Index)
+{
+	if (!MediaSourceLabels.IsValidIndex(Index))
+	{
+		if (MediaSources.IsValidIndex(Index))
+		{
+			if (MediaSources[Index])
+			{
+				return MediaSources[Index]->GetName();
+			}
+			else
+			{
+				return TEXT("");
+			}
+		}
+		
+		return TEXT("");
+	}
+	
+	return MediaSourceLabels[Index];
+}
+
+void UMediaProfile::SetLabelForMediaSource(int32 Index, const FString& NewLabel)
+{
+	if (!MediaSourceLabels.IsValidIndex(Index))
+	{
+		if (!MediaSources.IsValidIndex(Index))
+		{
+			return;
+		}
+		
+		// In some cases, we may not have labels for all media sources yet (such as an older media profile being edited)
+		// In that case, add all the missing labels now
+		for (int32 MissingIdx = MediaSourceLabels.Num(); MissingIdx < MediaSources.Num(); ++MissingIdx)
+		{
+			if (MediaSources[MissingIdx])
+			{
+				MediaSourceLabels.Add(MediaSources[MissingIdx]->GetName());
+			}
+			else
+			{
+				MediaSourceLabels.Add(TEXT(""));
+			}
+		}
+	}
+	
+	MediaSourceLabels[Index] = NewLabel;
+	
+	FPropertyChangedEvent PropertyChangedEvent(StaticClass()->FindPropertyByName(GET_MEMBER_NAME_CHECKED(UMediaProfile, MediaSourceLabels)), EPropertyChangeType::ValueSet);
+	FCoreUObjectDelegates::OnObjectPropertyChanged.Broadcast(this, PropertyChangedEvent);
+}
+#endif //WITH_EDITOR
+
 UMediaOutput* UMediaProfile::GetMediaOutput(int32 Index) const
 {
 	if (MediaOutputs.IsValidIndex(Index))
@@ -97,15 +209,56 @@ UMediaOutput* UMediaProfile::GetMediaOutput(int32 Index) const
 	return nullptr;
 }
 
+void UMediaProfile::SetMediaOutput(int32 Index, UMediaOutput* InMediaOutput)
+{
+	if (MediaOutputs.IsValidIndex(Index))
+	{
+		FString CurrentMediaOutputName = MediaOutputs[Index] ? MediaOutputs[Index]->GetName() : TEXT("");
+		MediaOutputs[Index] = InMediaOutput;
+		
+#if WITH_EDITOR
+		FPropertyChangedEvent PropertyChangedEvent(StaticClass()->FindPropertyByName(GET_MEMBER_NAME_CHECKED(UMediaProfile, MediaOutputs)), EPropertyChangeType::ValueSet);
+		FCoreUObjectDelegates::OnObjectPropertyChanged.Broadcast(this, PropertyChangedEvent);
+		
+		// If the original label for the media output was the output's object name (indicating the user has not changed it)
+		// update the label to the new output's object name
+		if (!MediaOutputLabels.IsValidIndex(Index) ||GetLabelForMediaOutput(Index) == CurrentMediaOutputName)
+		{
+			SetLabelForMediaOutput(Index, InMediaOutput->GetName());
+		}
+#endif
+	}
+}
 
 int32 UMediaProfile::NumMediaOutputs() const
 {
 	return MediaOutputs.Num();
 }
 
+int32 UMediaProfile::FindMediaOutputIndex(UMediaOutput* MediaOutput) const
+{
+	return MediaOutputs.Find(MediaOutput);
+}
+
 void UMediaProfile::AddMediaOutput(UMediaOutput* MediaOutput)
 {
 	MediaOutputs.Add(MediaOutput);
+
+#if WITH_EDITOR
+	if (MediaOutputLabels.Num() < MediaOutputs.Num())
+	{
+		for (int32 Index = MediaOutputLabels.Num(); Index < MediaOutputs.Num(); ++Index)
+		{
+			FString Name = MediaOutputs[Index] ? MediaOutputs[Index]->GetName() : TEXT("");
+			MediaOutputLabels.Add(Name);
+		}
+	}
+	else
+	{
+		FString Name = MediaOutput ? MediaOutput->GetName() : TEXT("");
+		MediaOutputLabels.Add(Name);
+	}
+#endif
 }
 
 bool UMediaProfile::RemoveMediaOutput(int32 Index)
@@ -116,6 +269,14 @@ bool UMediaProfile::RemoveMediaOutput(int32 Index)
 	}
 
 	MediaOutputs.RemoveAt(Index, EAllowShrinking::No);
+
+#if WITH_EDITOR
+	if (MediaOutputLabels.IsValidIndex(Index))
+	{
+		MediaOutputLabels.RemoveAt(Index, EAllowShrinking::No);
+	}
+#endif
+	
 	return true;
 }
 
@@ -129,10 +290,73 @@ bool UMediaProfile::MoveMediaOutput(int32 CurrentIndex, int32 DestIndex)
 	TObjectPtr<UMediaOutput> MediaOutput = MediaOutputs[CurrentIndex];
 	MediaOutputs.RemoveAt(CurrentIndex, EAllowShrinking::No);
 	MediaOutputs.Insert(MediaOutput, DestIndex);
+
+#if WITH_EDITOR
+	if (MediaOutputLabels.IsValidIndex(CurrentIndex))
+	{
+		FString OutputLabel = MediaOutputLabels[CurrentIndex];
+		MediaOutputLabels.RemoveAt(CurrentIndex, EAllowShrinking::No);
+		MediaOutputLabels.Insert(OutputLabel, DestIndex);
+	}
+#endif
 	
 	return true;
 }
 
+#if WITH_EDITOR
+FString UMediaProfile::GetLabelForMediaOutput(int32 Index)
+{
+	if (!MediaOutputLabels.IsValidIndex(Index))
+	{
+		if (MediaOutputs.IsValidIndex(Index))
+		{
+			if (MediaOutputs[Index])
+			{
+				return MediaOutputs[Index]->GetName();
+			}
+			else
+			{
+				return TEXT("");
+			}
+		}
+		
+		return TEXT("");
+	}
+	
+	return MediaOutputLabels[Index];
+}
+
+void UMediaProfile::SetLabelForMediaOutput(int32 Index, const FString& NewLabel)
+{
+	if (!MediaOutputLabels.IsValidIndex(Index))
+	{
+		if (!MediaOutputs.IsValidIndex(Index))
+		{
+			return;
+		}
+		
+		// In some cases, we may not have labels for all media outputs yet (such as an older media profile being edited)
+		// In that case, add all the missing labels now
+		for (int32 MissingIdx = MediaOutputLabels.Num(); MissingIdx < MediaOutputs.Num(); ++MissingIdx)
+		{
+			if (MediaOutputs[MissingIdx])
+			{
+				MediaOutputLabels.Add(MediaOutputs[MissingIdx]->GetName());
+			}
+			else
+			{
+				MediaOutputLabels.Add(TEXT(""));
+			}
+		}
+	}
+	
+	MediaOutputLabels[Index] = NewLabel;
+
+	FPropertyChangedEvent PropertyChangedEvent(StaticClass()->FindPropertyByName(GET_MEMBER_NAME_CHECKED(UMediaProfile, MediaOutputLabels)), EPropertyChangeType::ValueSet);
+	FCoreUObjectDelegates::OnObjectPropertyChanged.Broadcast(this, PropertyChangedEvent);
+}
+#endif //WITH_EDITOR
+
 UTimecodeProvider* UMediaProfile::GetTimecodeProvider() const
 {
 	return bOverrideTimecodeProvider ? TimecodeProvider : nullptr;
diff --git a/Engine/Plugins/Media/MediaFrameworkUtilities/Source/MediaFrameworkUtilities/Public/MediaAssets/ProxyMediaOutput.h b/Engine/Plugins/Media/MediaFrameworkUtilities/Source/MediaFrameworkUtilities/Public/MediaAssets/ProxyMediaOutput.h
--- a/Engine/Plugins/Media/MediaFrameworkUtilities/Source/MediaFrameworkUtilities/Public/MediaAssets/ProxyMediaOutput.h
+++ b/Engine/Plugins/Media/MediaFrameworkUtilities/Source/MediaFrameworkUtilities/Public/MediaAssets/ProxyMediaOutput.h
@@ -79,6 +79,7 @@ public:
 
 #if WITH_EDITOR
 	virtual FString GetDescriptionString() const override;
+	virtual void GetDetailsPanelInfoElements(TArray<FInfoElement>& OutInfoElements) const override;
 #endif //WITH_EDITOR
 	
 protected:
@@ -92,5 +93,6 @@ private:
 	mutable bool bRequestedSizeGuard;
 	mutable bool bRequestedPixelFormatGuard;
 	mutable bool bConfigurationTextGuard;
+	mutable bool bGetInfoElementsGuard;
 	mutable bool bCreateMediaCaptureImplGuard;
 };
diff --git a/Engine/Plugins/Media/MediaFrameworkUtilities/Source/MediaFrameworkUtilities/Public/MediaAssets/ProxyMediaSource.h b/Engine/Plugins/Media/MediaFrameworkUtilities/Source/MediaFrameworkUtilities/Public/MediaAssets/ProxyMediaSource.h
--- a/Engine/Plugins/Media/MediaFrameworkUtilities/Source/MediaFrameworkUtilities/Public/MediaAssets/ProxyMediaSource.h
+++ b/Engine/Plugins/Media/MediaFrameworkUtilities/Source/MediaFrameworkUtilities/Public/MediaAssets/ProxyMediaSource.h
@@ -77,6 +77,7 @@ public:
 
 #if WITH_EDITOR
 	virtual FString GetDescriptionString() const override;
+	virtual void GetDetailsPanelInfoElements(TArray<FInfoElement>& OutInfoElements) const override;
 #endif //WITH_EDITOR
 	
 public:
@@ -96,6 +97,7 @@ private:
 	mutable bool bUrlGuard;
 	mutable bool bValidateGuard;
 	mutable bool bConfigurationTextGuard;
+	mutable bool bGetInfoElementsGuard;
 	mutable bool bLeafMediaSource;
 	mutable bool bMediaOptionGuard;
 
diff --git a/Engine/Plugins/Media/MediaFrameworkUtilities/Source/MediaFrameworkUtilities/Public/Profile/MediaProfile.h b/Engine/Plugins/Media/MediaFrameworkUtilities/Source/MediaFrameworkUtilities/Public/Profile/MediaProfile.h
--- a/Engine/Plugins/Media/MediaFrameworkUtilities/Source/MediaFrameworkUtilities/Public/Profile/MediaProfile.h
+++ b/Engine/Plugins/Media/MediaFrameworkUtilities/Source/MediaFrameworkUtilities/Public/Profile/MediaProfile.h
@@ -47,6 +47,16 @@ protected:
 	UPROPERTY(EditAnywhere, Instanced, Category="Genlock", meta=(EditCondition="bOverrideCustomTimeStep"))
 	TObjectPtr<UEngineCustomTimeStep> CustomTimeStep;
 
+#if WITH_EDITORONLY_DATA
+	/** Stores user defined labels for the profile's media sources, so that users can organize and identify their media sources easier */
+	UPROPERTY()
+	TArray<FString> MediaSourceLabels;
+
+	/** Stores user defined labels for the profile's media outputs, so that users can organize and identify their media outputs easier */
+	UPROPERTY()
+	TArray<FString> MediaOutputLabels;
+#endif
+	
 public:
 #if WITH_EDITORONLY_DATA
 	/**
@@ -65,11 +75,25 @@ public:
 	 */
 	UMediaSource* GetMediaSource(int32 Index) const;
 
+	/**
+	 * Sets the media source at the specified index in the profile's list of sources
+	 * @param Index The index of the source to set
+	 * @param InMediaSource The new media source to set
+	 */
+	void SetMediaSource(int32 Index, UMediaSource* InMediaSource);
+	
 	/**
 	 * Get the number of media source.
 	 */
 	int32 NumMediaSources() const;
 
+	/**
+	 * Gets the index of the specified media source if it was found in the media profile's list of sources
+	 * @param MediaSource The media source to get the index of
+	 * @return The index of the source in the profile's sources list, or INDEX_NONE if it was not found
+	 */
+	int32 FindMediaSourceIndex(UMediaSource* MediaSource) const;
+	
 	/**
 	 * Adds the specified media source to the profile's list of media sources
 	 * @param MediaSource The media source to add
@@ -90,6 +114,22 @@ public:
 	 * @return true if the move occurred, otherwise false
 	 */
 	bool MoveMediaSource(int32 CurrentIndex, int32 DestIndex);
+
+#if WITH_EDITOR
+	/**
+	 * Gets the label for a media source
+	 * @param Index The index of the media source
+	 * @return The label that has been set for the media source
+	 */
+	FString GetLabelForMediaSource(int32 Index);
+
+	/**
+	 * Sets the label for a media source
+	 * @param Index The index of the media source
+	 * @param NewLabel The label to set for the media source
+	 */
+	void SetLabelForMediaSource(int32 Index, const FString& NewLabel);
+#endif //WITH_EDITOR
 	
 	/**
 	 * Get the media output for the selected proxy.
@@ -98,11 +138,25 @@ public:
 	 */
 	UMediaOutput* GetMediaOutput(int32 Index) const;
 
+	/**
+		 * Sets the media output at the specified index in the profile's list of outputs
+		 * @param Index The index of the output to set
+		 * @param InMediaOutput The new media output to set
+		 */
+	void SetMediaOutput(int32 Index, UMediaOutput* InMediaOutput);
+	
 	/**
 	 * Get the number of media output.
 	 */
 	int32 NumMediaOutputs() const;
 
+	/**
+	 * Gets the index of the specified media output if it was found in the media profile's list of outputs
+	 * @param MediaOutput The media output to get the index of
+	 * @return The index of the output in the profile's output list, or INDEX_NONE if it was not found
+	 */
+	int32 FindMediaOutputIndex(UMediaOutput* MediaOutput) const;
+	
 	/**
 	 * Adds the specified media output to the profile's list of media outputs
 	 * @param MediaOutput The media output to add
@@ -124,6 +178,22 @@ public:
 	 */
 	bool MoveMediaOutput(int32 CurrentIndex, int32 DestIndex);
 	
+#if WITH_EDITOR
+	/**
+	 * Gets the label for a media output
+	 * @param Index The index of the media output
+	 * @return The label that has been set for the media output
+	 */
+	FString GetLabelForMediaOutput(int32 Index);
+
+	/**
+	 * Sets the label for a media output
+	 * @param Index The index of the media output
+	 * @param NewLabel The label to set for the media output
+	 */
+	void SetLabelForMediaOutput(int32 Index, const FString& NewLabel);
+#endif //WITH_EDITOR
+	
 	/**
 	 * Get the timecode provider.
 	 *
diff --git a/Engine/Plugins/Media/MediaFrameworkUtilities/Source/MediaFrameworkUtilitiesEditor/Private/AssetEditor/DummyMediaObject.h b/Engine/Plugins/Media/MediaFrameworkUtilities/Source/MediaFrameworkUtilitiesEditor/Private/AssetEditor/DummyMediaObject.h
new file mode 100644
--- /dev/null
+++ b/Engine/Plugins/Media/MediaFrameworkUtilities/Source/MediaFrameworkUtilitiesEditor/Private/AssetEditor/DummyMediaObject.h
@@ -0,0 +1,39 @@
+﻿// Copyright Epic Games, Inc. All Rights Reserved.
+
+#pragma once
+
+#include "CoreMinimal.h"
+#include "MediaOutput.h"
+#include "MediaSource.h"
+
+#include "DummyMediaObject.generated.h"
+
+/** Placeholder to use in the details panel when the media profile has a null media source in its list of media sources */
+UCLASS(HideCategories=("Media|MediaSource"))
+class UDummyMediaSource : public UMediaSource
+{
+	GENERATED_BODY()
+
+public:
+	virtual FString GetUrl() const override { return TEXT(""); }
+	virtual bool Validate() const override { return false; }
+	
+public:
+	/** Index of the media source in the media profile's media sources list that this dummy represents */
+	int32 MediaProfileIndex = INDEX_NONE;
+};
+
+/** Placeholder to use in the details panel when the media profile has a null media output in its list of media outputs */
+UCLASS(HideCategories=("Output", "Media|Output"))
+class UDummyMediaOutput : public UMediaOutput
+{
+	GENERATED_BODY()
+
+public:
+	virtual FIntPoint GetRequestedSize() const override { return FIntPoint::ZeroValue; }
+	virtual EPixelFormat GetRequestedPixelFormat() const override { return EPixelFormat::PF_Unknown; }
+	
+public:
+	/** Index of the media output in the media profile's media outputs list that this dummy represents */
+	int32 MediaProfileIndex = INDEX_NONE;
+};
\ No newline at end of file
diff --git a/Engine/Plugins/Media/MediaFrameworkUtilities/Source/MediaFrameworkUtilitiesEditor/Private/AssetEditor/MediaProfileEditor.cpp b/Engine/Plugins/Media/MediaFrameworkUtilities/Source/MediaFrameworkUtilitiesEditor/Private/AssetEditor/MediaProfileEditor.cpp
--- a/Engine/Plugins/Media/MediaFrameworkUtilities/Source/MediaFrameworkUtilitiesEditor/Private/AssetEditor/MediaProfileEditor.cpp
+++ b/Engine/Plugins/Media/MediaFrameworkUtilities/Source/MediaFrameworkUtilitiesEditor/Private/AssetEditor/MediaProfileEditor.cpp
@@ -3,11 +3,18 @@
 
 #include "MediaProfileEditor.h"
 
+#include "Engine/TextureRenderTarget2D.h"
+#include "LevelEditorViewport.h"
 #include "PropertyEditorModule.h"
+#include "SMediaProfileDetailsPanel.h"
 #include "SMediaProfileSourcesTreeView.h"
+#include "CaptureTab/SMediaFrameworkCapture.h"
+#include "GameFramework/WorldSettings.h"
+#include "Misc/TransactionObjectEvent.h"
 #include "Modules/ModuleManager.h"
 #include "Profile/IMediaProfileManager.h"
 #include "Profile/MediaProfile.h"
+#include "Slate/SceneViewport.h"
 #include "Widgets/Docking/SDockTab.h"
 
 #define LOCTEXT_NAMESPACE "MediaProfileEditor"
@@ -18,6 +25,121 @@ const FName FMediaProfileEditor::MediaTreeTabId =  FName(TEXT("MediaTree"));
 const FName FMediaProfileEditor::DetailsTabId =  FName(TEXT("Details"));
 const FName FMediaProfileEditor::TimecodeTabId =  FName(TEXT("Timecode"));
 
+namespace MediaProfileEditor
+{	
+	/**
+	 * Determines the target size from the media output if desired size is specified, otherwise,
+	 * fallback to the default viewport capture resolution.
+	 */
+	FIntPoint GetTargetSize(UMediaOutput* InMediaOutput)
+	{
+		const FIntPoint TargetSize = InMediaOutput->GetRequestedSize();
+		if (TargetSize != UMediaOutput::RequestCaptureSourceSize)
+		{
+			return TargetSize;
+		}
+		
+		return GetMutableDefault<UMediaFrameworkMediaCaptureSettings>()->DefaultViewportCaptureSize;
+	}
+}
+
+/*
+ * FMediaFrameworkCaptureLevelEditorViewportClient
+ * Like FLevelEditorViewportClient but always use the PlayWorld if available. Do not support bSimulated
+ */
+class FMediaFrameworkCaptureLevelEditorViewportClient : public FLevelEditorViewportClient
+{
+private:
+	using Super = FLevelEditorViewportClient;
+
+public:
+	static TSharedPtr<FMediaFrameworkCaptureLevelEditorViewportClient> CreateViewportClientForCapture(const FMediaFrameworkCaptureCameraViewportCameraOutputInfo& InViewportOutputInfo)
+	{
+		TSharedPtr<FMediaFrameworkCaptureLevelEditorViewportClient> NewViewportClient = MakeShared<FMediaFrameworkCaptureLevelEditorViewportClient>(nullptr, InViewportOutputInfo.ViewMode);
+
+		// Preview viewports never be a listener
+		NewViewportClient->bSetListenerPosition = false;
+
+		// Default to "game" show flags for camera previews
+		NewViewportClient->EngineShowFlags = FEngineShowFlags(ESFIM_Game);
+		NewViewportClient->LastEngineShowFlags = FEngineShowFlags(ESFIM_Editor);
+
+		NewViewportClient->ViewportType = LVT_Perspective;
+		NewViewportClient->bDrawAxes = false;
+		NewViewportClient->bDisableInput = true;
+		NewViewportClient->SetAllowCinematicControl(false);
+		NewViewportClient->VisibilityDelegate.BindLambda([] {return true; });
+		NewViewportClient->AutoSetPIE();
+		
+		return NewViewportClient;
+	}
+	
+	FMediaFrameworkCaptureLevelEditorViewportClient(const TSharedPtr<SLevelViewport>& InLevelViewport, EViewModeIndex InViewModeIndex)
+		: FLevelEditorViewportClient(InLevelViewport)
+		, ViewModeIndex(InViewModeIndex)
+	{}
+
+	virtual ~FMediaFrameworkCaptureLevelEditorViewportClient() override
+	{
+		if (GEditor)
+		{
+			FWorldContext* PIEWorldContext = GEditor->GetPIEWorldContext();
+			if (PIEWorldContext)
+			{
+				RemoveReferenceToWorldContext(*PIEWorldContext);
+			}
+		}
+	}
+
+	virtual UWorld* GetWorld() const override
+	{
+		if (GEditor->PlayWorld)
+		{
+			return GEditor->PlayWorld;
+		}
+		return Super::GetWorld();
+	}
+
+	void SetPIE(bool bInIsPIE)
+	{
+		FWorldContext* PIEWorldContext = GEditor->GetPIEWorldContext();
+		if (bInIsPIE)
+		{
+			if (PIEWorldContext)
+			{
+				RemoveReferenceToWorldContext(GEditor->GetEditorWorldContext());
+				SetReferenceToWorldContext(*PIEWorldContext);
+			}
+		}
+		else
+		{
+			if (PIEWorldContext)
+			{
+				RemoveReferenceToWorldContext(*PIEWorldContext);
+			}
+			SetReferenceToWorldContext(GEditor->GetEditorWorldContext());
+		}
+		SetViewMode(ViewModeIndex);
+		SetRealtime(true);
+	}
+
+	void AutoSetPIE()
+	{
+		if (GEditor->PlayWorld)
+		{
+			SetPIE(true);
+		}
+		else
+		{
+			SetViewMode(ViewModeIndex);
+			SetRealtime(true);
+		}
+	}
+
+private:
+	EViewModeIndex ViewModeIndex;
+};
+
 TSharedRef<FMediaProfileEditor> FMediaProfileEditor::CreateMediaProfileEditor(const EToolkitMode::Type Mode, const TSharedPtr<IToolkitHost>& InitToolkitHost, UMediaProfile* InMediaProfile)
 {
 	TSharedRef<FMediaProfileEditor> NewEditor = MakeShared<FMediaProfileEditor>();
@@ -30,10 +152,13 @@ void FMediaProfileEditor::Initialize(const EToolkitMode::Type Mode, const TShare
 	GEditor->GetEditorSubsystem<UAssetEditorSubsystem>()->CloseOtherEditors(InMediaProfile, this);
 	MediaProfileBeingEdited = InMediaProfile;
 
-	FPropertyEditorModule& PropertyModule = FModuleManager::LoadModuleChecked<FPropertyEditorModule>("PropertyEditor");
-	FDetailsViewArgs DatabaseDetailsArgs;
-	DetailsView = PropertyModule.CreateDetailView(DatabaseDetailsArgs);
-	DetailsView->SetObject(InMediaProfile);
+	FEditorDelegates::PostPIEStarted.AddSP(this, &FMediaProfileEditor::OnPostPIEStarted);
+	FEditorDelegates::PrePIEEnded.AddSP(this, &FMediaProfileEditor::OnPrePIEEnded);
+	FEditorDelegates::MapChange.AddSP(this, &FMediaProfileEditor::OnMapChange);
+	GEngine->OnLevelActorDeleted().AddSP(this, &FMediaProfileEditor::OnLevelActorsRemoved);
+	FEditorDelegates::OnAssetsDeleted.AddSP(this, &FMediaProfileEditor::OnAssetsDeleted);
+	FCoreUObjectDelegates::OnPreObjectPropertyChanged.AddSP(this, &FMediaProfileEditor::OnObjectPreEditChange);
+	FCoreUObjectDelegates::OnObjectModified.AddSP(this, &FMediaProfileEditor::OnObjectModified);
 	
 	const TSharedRef<FTabManager::FLayout> StandaloneDefaultLayout = FTabManager::NewLayout("Standalone_MediaProfileEditor_Layout")
 		->AddArea
@@ -81,6 +206,19 @@ void FMediaProfileEditor::Initialize(const EToolkitMode::Type Mode, const TShare
 	IMediaProfileManager::Get().SetCurrentMediaProfile(InMediaProfile);
 }
 
+FMediaProfileEditor::~FMediaProfileEditor()
+{
+	FCoreUObjectDelegates::OnObjectModified.RemoveAll(this);
+	FCoreUObjectDelegates::OnPreObjectPropertyChanged.RemoveAll(this);
+	FEditorDelegates::OnAssetsDeleted.RemoveAll(this);
+	GEngine->OnLevelActorDeleted().RemoveAll(this);
+	FEditorDelegates::MapChange.RemoveAll(this);
+	FEditorDelegates::PrePIEEnded.RemoveAll(this);
+	FEditorDelegates::PostPIEStarted.RemoveAll(this);
+	
+	StopCapturingAllMediaOutputs();
+}
+
 void FMediaProfileEditor::RegisterTabSpawners(const TSharedRef<class FTabManager>& InTabManager)
 {
 	WorkspaceMenuCategory = InTabManager->AddLocalWorkspaceMenuCategory(LOCTEXT("WorkspaceMenuCategory", "Media Profile Editor"));
@@ -145,6 +283,219 @@ FLinearColor FMediaProfileEditor::GetWorldCentricTabColorScale() const
 	return FLinearColor::White;
 }
 
+bool FMediaProfileEditor::IsMediaOutputCapturing(UMediaOutput* InMediaOutput) const
+{
+	return LiveCaptures.ContainsByPredicate([InMediaOutput](const FLiveMediaCapture& LiveCapture)
+	{
+		return LiveCapture.MediaOutput.IsValid() && LiveCapture.MediaOutput.Get() == InMediaOutput;
+	});
+}
+
+bool FMediaProfileEditor::CanMediaOutputCapture(UMediaOutput* InMediaOutput) const
+{
+	UMediaFrameworkWorldSettingsAssetUserData* CaptureSettings = FindOrAddMediaFrameworkCaptureSettings();
+	if (!CaptureSettings)
+	{
+		return false;
+	}
+
+	// Return true if there is at least one valid capture configuration for the specified media output
+	bool bHasValidCaptureSettings = false;
+	for (const FMediaFrameworkCaptureCameraViewportCameraOutputInfo& ViewportOutputInfo : CaptureSettings->ViewportCaptures)
+	{
+		if (ViewportOutputInfo.MediaOutput != InMediaOutput)
+		{
+			continue;
+		}
+
+		const bool bHasValidCamera = ViewportOutputInfo.LockedActors.ContainsByPredicate([](const TLazyObjectPtr<AActor>& ActorRef)
+		{
+			return ActorRef.IsValid();
+		});
+		
+		if (bHasValidCamera)
+		{
+			bHasValidCaptureSettings = true;
+			break;
+		}
+	}
+	
+	for (const FMediaFrameworkCaptureRenderTargetCameraOutputInfo& RenderTargetOutputInfo : CaptureSettings->RenderTargetCaptures)
+	{
+		if (RenderTargetOutputInfo.MediaOutput != InMediaOutput)
+		{
+			continue;
+		}
+
+		if (RenderTargetOutputInfo.RenderTarget)
+		{
+			bHasValidCaptureSettings = true;
+			break;
+		}
+	}
+
+	return bHasValidCaptureSettings;
+}
+
+void FMediaProfileEditor::StartCapturingMediaOutput(UMediaOutput* InMediaOutput)
+{
+	UMediaFrameworkWorldSettingsAssetUserData* CaptureSettings = FindOrAddMediaFrameworkCaptureSettings();
+	if (!CaptureSettings)
+	{
+		return;
+	}
+
+	if (!InMediaOutput)
+	{
+		return;
+	}
+	
+	bool bStartedCapture = false;
+	for (const FMediaFrameworkCaptureCameraViewportCameraOutputInfo& ViewportOutputInfo : CaptureSettings->ViewportCaptures)
+	{
+		if (ViewportOutputInfo.MediaOutput != InMediaOutput)
+		{
+			continue;
+		}
+		
+		FLiveMediaCapture NewLiveCapture;
+		NewLiveCapture.MediaOutput = InMediaOutput;
+		NewLiveCapture.MediaCapture = TStrongObjectPtr(InMediaOutput->CreateMediaCapture());
+
+		if (NewLiveCapture.MediaCapture.IsValid())
+		{
+			// TODO: Hook this up to a slate widget when the media viewport tab is created
+			NewLiveCapture.ViewportClient = FMediaFrameworkCaptureLevelEditorViewportClient::CreateViewportClientForCapture(ViewportOutputInfo);
+			NewLiveCapture.Viewport = MakeShared<FSceneViewport>(NewLiveCapture.ViewportClient.Get(), nullptr);
+			NewLiveCapture.ViewportClient->Viewport = NewLiveCapture.Viewport.Get();
+
+			const FIntPoint TargetSize = MediaProfileEditor::GetTargetSize(InMediaOutput);
+			NewLiveCapture.Viewport->SetFixedViewportSize(TargetSize.X, TargetSize.Y);
+			
+			if (NewLiveCapture.MediaCapture->CaptureSceneViewport(NewLiveCapture.Viewport, ViewportOutputInfo.CaptureOptions))
+			{
+				LiveCaptures.Add(NewLiveCapture);
+				bStartedCapture = true;
+			}
+			else
+			{
+				NewLiveCapture.ViewportClient->Viewport = nullptr;
+			}
+		}
+	}
+
+	for (const FMediaFrameworkCaptureRenderTargetCameraOutputInfo& RenderTargetOutputInfo : CaptureSettings->RenderTargetCaptures)
+	{
+		if (RenderTargetOutputInfo.MediaOutput != InMediaOutput)
+		{
+			continue;
+		}
+		
+		FLiveMediaCapture NewLiveCapture;
+		NewLiveCapture.MediaOutput = InMediaOutput;
+		NewLiveCapture.MediaCapture = TStrongObjectPtr(InMediaOutput->CreateMediaCapture());
+
+		if (NewLiveCapture.MediaCapture.IsValid())
+		{
+			// TODO: When the media viewport tab is configured, hook up the capture to a widget
+			if (NewLiveCapture.MediaCapture->CaptureTextureRenderTarget2D(RenderTargetOutputInfo.RenderTarget.Get(), RenderTargetOutputInfo.CaptureOptions))
+			{
+				LiveCaptures.Add(NewLiveCapture);
+				bStartedCapture = true;
+			}
+		}
+	}
+
+	if (bStartedCapture)
+	{
+		InMediaOutput->OnOutputModified().AddSP(this, &FMediaProfileEditor::OnMediaOutputModified);
+	}
+}
+
+void FMediaProfileEditor::StopCapturingMediaOutput(UMediaOutput* InMediaOutput)
+{
+	for (FLiveMediaCapture& Capture : LiveCaptures)
+	{
+		if (Capture.MediaOutput.Get() != InMediaOutput)
+		{
+			continue;
+		}
+
+		if (Capture.MediaCapture.IsValid())
+		{
+			Capture.MediaCapture->StopCapture(false);
+		}
+
+		if (Capture.MediaOutput.IsValid())
+		{
+			Capture.MediaOutput->OnOutputModified().RemoveAll(this);
+		}
+
+		if (Capture.ViewportClient.IsValid())
+		{
+			Capture.ViewportClient->Viewport = nullptr;
+		}
+	}
+
+	LiveCaptures.RemoveAll([InMediaOutput](const FLiveMediaCapture& Capture)
+	{
+		return Capture.MediaOutput.Get() == InMediaOutput;
+	});
+}
+
+void FMediaProfileEditor::StopCapturingAllMediaOutputs()
+{
+	for (FLiveMediaCapture& Capture : LiveCaptures)
+	{
+		if (Capture.MediaCapture.IsValid())
+		{
+			Capture.MediaCapture->StopCapture(false);
+		}
+
+		if (Capture.MediaOutput.IsValid())
+		{
+			Capture.MediaOutput->OnOutputModified().RemoveAll(this);
+		}
+
+		if (Capture.ViewportClient.IsValid())
+		{
+			Capture.ViewportClient->Viewport = nullptr;
+		}
+	}
+
+	LiveCaptures.Empty();
+}
+
+UMediaFrameworkWorldSettingsAssetUserData* FMediaProfileEditor::FindOrAddMediaFrameworkCaptureSettings()
+{
+	UMediaFrameworkWorldSettingsAssetUserData* Result = nullptr;
+
+	if (GetDefault<UMediaFrameworkEditorCaptureSettings>()->bSaveCaptureSetingsInWorld)
+	{
+		UWorld* World = GEditor ? GEditor->GetEditorWorldContext(false).World() : nullptr;
+		if (World)
+		{
+			AWorldSettings* WorldSetting = World ? World->GetWorldSettings() : nullptr;
+			if (WorldSetting)
+			{
+				Result = CastChecked<UMediaFrameworkWorldSettingsAssetUserData>(WorldSetting->GetAssetUserDataOfClass(UMediaFrameworkWorldSettingsAssetUserData::StaticClass()), ECastCheckedType::NullAllowed);
+
+				if (!Result)
+				{
+					Result = NewObject<UMediaFrameworkWorldSettingsAssetUserData>(WorldSetting);
+					WorldSetting->AddAssetUserData(Result);
+				}
+			}
+		}
+	}
+	else
+	{
+		Result = GetMutableDefault<UMediaFrameworkEditorCaptureSettings>();
+	}
+
+	return Result;
+}
+
 TSharedRef<SDockTab> FMediaProfileEditor::SpawnTab_MediaOutput(const FSpawnTabArgs& InArgs)
 {
 	return SNew(SDockTab)
@@ -160,6 +511,7 @@ TSharedRef<SDockTab> FMediaProfileEditor::SpawnTab_MediaTree(const FSpawnTabArgs
 		.Label(LOCTEXT("MediaTreeTabLabel", "Media"))
 		[
 			SNew(SMediaProfileSourcesTreeView, MediaProfileBeingEdited)
+			.OnSelectedMediaItemsChanged(this, &FMediaProfileEditor::OnSelectedMediaItemsChanged)
 		];
 }
 
@@ -168,7 +520,7 @@ TSharedRef<SDockTab> FMediaProfileEditor::SpawnTab_Details(const FSpawnTabArgs&
 	return SNew(SDockTab)
 		.Label(LOCTEXT("DetailsTabLabel", "Details"))
 		[
-			DetailsView.ToSharedRef()
+			SAssignNew(DetailsPanel, SMediaProfileDetailsPanel, SharedThis(this), MediaProfileBeingEdited)
 		];
 }
 
@@ -181,4 +533,163 @@ TSharedRef<SDockTab> FMediaProfileEditor::SpawnTab_Timecode(const FSpawnTabArgs&
 		];
 }
 
+void FMediaProfileEditor::OnSelectedMediaItemsChanged(const TArray<int32>& SelectedMediaSources, const TArray<int32>& SelectedMediaOutputs)
+{
+	if (!DetailsPanel.IsValid())
+	{
+		return;
+	}
+
+	DetailsPanel->SetSelectedMediaItems(SelectedMediaSources, SelectedMediaOutputs);
+}
+
+void FMediaProfileEditor::OnMediaOutputModified(UMediaOutput* InMediaOutput)
+{
+	if (!GetDefault<UMediaFrameworkEditorCaptureSettings>()->bAutoRestartCaptureOnChange)
+	{
+		return;
+	}
+	
+	const bool bIsCapturing = LiveCaptures.ContainsByPredicate([InMediaOutput](const FLiveMediaCapture& Capture)
+	{
+		return Capture.MediaOutput.Get() == InMediaOutput;
+	});
+
+	if (!bIsCapturing)
+	{
+		return;
+	}
+
+	StopCapturingMediaOutput(InMediaOutput);
+	OutputsToRestartCapture.Add(InMediaOutput);
+	
+	if (!RestartCapturesTimerHandle.IsValid())
+	{
+		RestartCapturesTimerHandle = GEditor->GetTimerManager()->SetTimerForNextTick(FTimerDelegate::CreateSP(this, &FMediaProfileEditor::OnRestartCapturesNextTick));
+	}
+}
+
+void FMediaProfileEditor::OnRestartCapturesNextTick()
+{
+	for (const TWeakObjectPtr<UMediaOutput>& MediaOutput : OutputsToRestartCapture)
+	{
+		if (MediaOutput.IsValid())
+		{
+			StartCapturingMediaOutput(MediaOutput.Get());
+		}
+	}
+
+	OutputsToRestartCapture.Empty();
+	RestartCapturesTimerHandle.Invalidate();
+}
+
+void FMediaProfileEditor::OnPostPIEStarted(bool bIsSimulating)
+{
+	for (FLiveMediaCapture& Capture : LiveCaptures)
+	{
+		if (Capture.ViewportClient.IsValid())
+		{
+			Capture.ViewportClient->SetPIE(true);
+		}
+	}
+}
+
+void FMediaProfileEditor::OnPrePIEEnded(bool bIsSimulating)
+{
+	for (FLiveMediaCapture& Capture : LiveCaptures)
+	{
+		if (Capture.ViewportClient.IsValid())
+		{
+			Capture.ViewportClient->SetPIE(false);
+		}
+	}
+}
+
+void FMediaProfileEditor::OnMapChange(uint32 InMapFlags)
+{
+	// New map might have different capture settings, so stop all captures and notify of settings changes
+	StopCapturingAllMediaOutputs();
+	OnCaptureSettingsObjectChanged.Broadcast();
+}
+
+void FMediaProfileEditor::OnLevelActorsRemoved(AActor* InActor)
+{
+	if (!LiveCaptures.Num())
+	{
+		return;
+	}
+	
+	UMediaFrameworkWorldSettingsAssetUserData* CaptureSettings = FindOrAddMediaFrameworkCaptureSettings();
+	if (!CaptureSettings)
+	{
+		return;
+	}
+	
+	// If the removed actor is referenced by one of the live captures (e.g. a camera in a viewport capture) we need to stop the capture
+	for (const FMediaFrameworkCaptureCameraViewportCameraOutputInfo& ViewportOutputInfo : CaptureSettings->ViewportCaptures)
+	{
+		if (ViewportOutputInfo.LockedActors.Contains(InActor))
+		{
+			StopCapturingMediaOutput(ViewportOutputInfo.MediaOutput);
+		}
+	}
+}
+
+void FMediaProfileEditor::OnAssetsDeleted(const TArray<UClass*>& DeletedAssetClasses)
+{
+	if (!LiveCaptures.Num())
+	{
+		return;
+	}
+	
+	UMediaFrameworkWorldSettingsAssetUserData* CaptureSettings = FindOrAddMediaFrameworkCaptureSettings();
+	if (!CaptureSettings)
+	{
+		return;
+	}
+
+	bool bContainsRenderTargets = DeletedAssetClasses.ContainsByPredicate([](const UClass* Class)
+	{
+		return Class->IsChildOf<UTextureRenderTarget2D>();
+	});
+	
+	if (bContainsRenderTargets)
+	{
+		// If the deleted asset is a render target, we need to stop all render target captures in case the deleted render target is used in one of them
+		for (const FMediaFrameworkCaptureRenderTargetCameraOutputInfo& RenderTargetOutputInfo : CaptureSettings->RenderTargetCaptures)
+		{
+			StopCapturingMediaOutput(RenderTargetOutputInfo.MediaOutput);
+		}
+	}
+}
+
+void FMediaProfileEditor::OnObjectPreEditChange(UObject* Object, const FEditPropertyChain& PropertyChain)
+{
+	// If the capture settings are about to be changed, stop all captures
+	UMediaFrameworkWorldSettingsAssetUserData* CaptureSettings = FindOrAddMediaFrameworkCaptureSettings();
+	if (Object == CaptureSettings)
+	{
+		StopCapturingAllMediaOutputs();
+	}
+}
+
+void FMediaProfileEditor::OnObjectModified(UObject* Object)
+{
+	const UMediaFrameworkEditorCaptureSettings* GlobalCaptureSettings = GetDefault<UMediaFrameworkEditorCaptureSettings>();
+	if (Object == GlobalCaptureSettings)
+	{
+		bool bCachedSaveCaptureSetingsInWorld = GlobalCaptureSettings->bSaveCaptureSetingsInWorld;
+		
+		// Defer invoking the delegate to ensure any modifications to the global settings complete
+		GEditor->GetTimerManager()->SetTimerForNextTick([this, bCachedSaveCaptureSetingsInWorld, GlobalCaptureSettings]()
+		{
+			// If the capture settings object has changed to or from the global settings, invoke the CaptureSettingsObjectChanged delegate
+			if (bCachedSaveCaptureSetingsInWorld != GlobalCaptureSettings->bSaveCaptureSetingsInWorld)
+			{
+				OnCaptureSettingsObjectChanged.Broadcast();
+			}
+		});
+	}
+}
+
 #undef LOCTEXT_NAMESPACE
diff --git a/Engine/Plugins/Media/MediaFrameworkUtilities/Source/MediaFrameworkUtilitiesEditor/Private/AssetEditor/MediaProfileEditor.h b/Engine/Plugins/Media/MediaFrameworkUtilities/Source/MediaFrameworkUtilitiesEditor/Private/AssetEditor/MediaProfileEditor.h
--- a/Engine/Plugins/Media/MediaFrameworkUtilities/Source/MediaFrameworkUtilitiesEditor/Private/AssetEditor/MediaProfileEditor.h
+++ b/Engine/Plugins/Media/MediaFrameworkUtilities/Source/MediaFrameworkUtilitiesEditor/Private/AssetEditor/MediaProfileEditor.h
@@ -5,7 +5,12 @@
 #include "CoreMinimal.h"
 #include "Toolkits/AssetEditorToolkit.h"
 
+class FMediaFrameworkCaptureLevelEditorViewportClient;
+class UMediaFrameworkWorldSettingsAssetUserData;
+class UMediaOutput;
+class UMediaCapture;
 class UMediaProfile;
+class SMediaProfileDetailsPanel;
 
 /**
  * Media profile editor 2.0, which replaces the old, simple editor with a more robust one
@@ -24,6 +29,8 @@ public:
 
 	void Initialize(const EToolkitMode::Type Mode, const TSharedPtr< IToolkitHost >& InitToolkitHost, UMediaProfile* InMediaProfile);
 
+	virtual ~FMediaProfileEditor() override;
+	
 	// IToolkit interface
 	virtual void RegisterTabSpawners(const TSharedRef<class FTabManager>& InTabManager) override;
 	virtual void UnregisterTabSpawners(const TSharedRef<class FTabManager>& InTabManager) override;
@@ -38,15 +45,70 @@ public:
 	virtual FLinearColor GetWorldCentricTabColorScale() const override;
 	// End of FAssetEditorToolkit
 
+	/** Determines if the specified media output is currently capturing */
+	bool IsMediaOutputCapturing(UMediaOutput* InMediaOutput) const;
+
+	/** Determines if the specified media output is properly configured for capturing */
+	bool CanMediaOutputCapture(UMediaOutput* InMediaOutput) const;
+	
+	/** Starts capturing the specified media output */
+	void StartCapturingMediaOutput(UMediaOutput* InMediaOutput);
+
+	/** Stops capture on any live captures for the specified media output */
+	void StopCapturingMediaOutput(UMediaOutput* InMediaOutput);
+
+	/** Stops capture on all live captures */
+	void StopCapturingAllMediaOutputs();
+	
+	/** Gets a valid media capture settings object if one is found, or attempts to add one to the world if one is not found */
+	static UMediaFrameworkWorldSettingsAssetUserData* FindOrAddMediaFrameworkCaptureSettings();
+
+	/** Gets the delegate that is raised when the capture settings object has been changed for any reason */
+	FSimpleMulticastDelegate& GetOnCaptureSettingsObjectChanged() { return OnCaptureSettingsObjectChanged; }
+	
 private:
 	TSharedRef<SDockTab> SpawnTab_MediaOutput(const FSpawnTabArgs& InArgs);
 	TSharedRef<SDockTab> SpawnTab_MediaTree(const FSpawnTabArgs& InArgs);
 	TSharedRef<SDockTab> SpawnTab_Details(const FSpawnTabArgs& InArgs);
 	TSharedRef<SDockTab> SpawnTab_Timecode(const FSpawnTabArgs& InArgs);
+	
+	void OnSelectedMediaItemsChanged(const TArray<int32>& SelectedMediaSources, const TArray<int32>& SelectedMediaOutputs);
+	void OnMediaOutputModified(UMediaOutput* InMediaOutput);
+	void OnRestartCapturesNextTick();
+
+
+	void OnPostPIEStarted(bool bIsSimulating);
+	void OnPrePIEEnded(bool bIsSimulating);
+	
+	void OnMapChange(uint32 InMapFlags);
+	void OnLevelActorsRemoved(AActor* InActor);
+	void OnAssetsDeleted(const TArray<UClass*>& DeletedAssetClasses);
+	void OnObjectPreEditChange(UObject* Object, const FEditPropertyChain& PropertyChain);
+	void OnObjectModified(UObject* Object);
+	
 private:
 	/** The media profile being edited by this editor */
 	TObjectPtr<UMediaProfile> MediaProfileBeingEdited;
 
-	/** Details view to display the media profile properties */
-	TSharedPtr<IDetailsView> DetailsView;
+	TSharedPtr<SMediaProfileDetailsPanel> DetailsPanel;
+
+	struct FLiveMediaCapture
+	{
+		TWeakObjectPtr<UMediaOutput> MediaOutput;
+		TStrongObjectPtr<UMediaCapture> MediaCapture;
+		TSharedPtr<FMediaFrameworkCaptureLevelEditorViewportClient> ViewportClient;
+		TSharedPtr<FSceneViewport> Viewport;
+	};
+	
+	/** List of any live media output captures that the media profile editor is managing */
+	TArray<FLiveMediaCapture> LiveCaptures;
+
+	/** List of media outputs whose capture needs to be restarted on the next engine tick */
+	TSet<TWeakObjectPtr<UMediaOutput>> OutputsToRestartCapture;
+
+	/** Raised when the object that stores the capture settings has changed */
+	FSimpleMulticastDelegate OnCaptureSettingsObjectChanged;
+	
+	/** Handle for callback to restart any modified captures on the nextx engine tick */
+	FTimerHandle RestartCapturesTimerHandle;
 };
diff --git a/Engine/Plugins/Media/MediaFrameworkUtilities/Source/MediaFrameworkUtilitiesEditor/Private/AssetEditor/MediaProfileMediaItemDetailCustomization.cpp b/Engine/Plugins/Media/MediaFrameworkUtilities/Source/MediaFrameworkUtilitiesEditor/Private/AssetEditor/MediaProfileMediaItemDetailCustomization.cpp
new file mode 100644
--- /dev/null
+++ b/Engine/Plugins/Media/MediaFrameworkUtilities/Source/MediaFrameworkUtilitiesEditor/Private/AssetEditor/MediaProfileMediaItemDetailCustomization.cpp
@@ -0,0 +1,664 @@
+﻿// Copyright Epic Games, Inc. All Rights Reserved.
+
+#include "MediaProfileMediaItemDetailCustomization.h"
+
+#include "IDetailChildrenBuilder.h"
+#include "IStructureDataProvider.h"
+#include "MediaFrameworkWorldSettingsAssetUserData.h"
+#include "Algo/Count.h"
+#include "UI/MediaFrameworkUtilitiesEditorStyle.h"
+
+#define LOCTEXT_NAMESPACE "FMediaProfileMediaItemDetailCustomization"
+
+FText FMediaProfileMediaSourceDetailCustomization::GetMediaItemTypeText() const
+{
+	return LOCTEXT("MediaSourceTypeText", "Media Source");
+}
+
+FText FMediaProfileMediaSourceDetailCustomization::GetMediaObjectLabel(UMediaSource* InMediaObject) const
+{
+	if (!MediaProfile.IsValid())
+	{
+		return FText::GetEmpty();
+	}
+
+	int32 Index = MediaProfile->FindMediaSourceIndex(InMediaObject);
+	if (InMediaObject->IsA<UDummyMediaSource>())
+	{
+		Index = Cast<UDummyMediaSource>(InMediaObject)->MediaProfileIndex;
+	}
+		
+	return FText::FromString(MediaProfile->GetLabelForMediaSource(Index));
+}
+
+void FMediaProfileMediaSourceDetailCustomization::SetMediaObjectLabel(UMediaSource* InMediaObject, const FText& InLabel)
+{
+	if (!MediaProfile.IsValid())
+	{
+		return;
+	}
+	
+	int32 Index = MediaProfile->FindMediaSourceIndex(InMediaObject);
+	if (InMediaObject->IsA<UDummyMediaSource>())
+	{
+		Index = Cast<UDummyMediaSource>(InMediaObject)->MediaProfileIndex;
+	}
+		
+	MediaProfile->SetLabelForMediaSource(Index, InLabel.ToString());
+}
+
+void FMediaProfileMediaSourceDetailCustomization::SetMediaObject(UMediaSource* InOriginalMediaObject, UMediaSource* InNewMediaObject)
+{
+	if (!MediaProfile.IsValid())
+    {
+    	return;
+    }
+    	
+	int32 Index = MediaProfile->FindMediaSourceIndex(InOriginalMediaObject);
+	if (InOriginalMediaObject->IsA<UDummyMediaSource>())
+	{
+		Index = Cast<UDummyMediaSource>(InOriginalMediaObject)->MediaProfileIndex;
+	}
+		
+	MediaProfile->SetMediaSource(Index, InNewMediaObject);
+}
+
+class SMediaCaptureMethodComboButton : public SComboButton
+{
+public:
+	SLATE_BEGIN_ARGS(SMediaCaptureMethodComboButton) { }
+		SLATE_EVENT(FSimpleDelegate, OnMethodChanged)
+	SLATE_END_ARGS()
+
+	void Construct(const FArguments& InArgs, TArray<TWeakObjectPtr<UObject>>& InObjects)
+	{
+		OnMethodChanged = InArgs._OnMethodChanged;
+		Objects = InObjects;
+		
+		SComboButton::Construct(SComboButton::FArguments()
+			.OnGetMenuContent_Lambda([this]
+			{
+				FMenuBuilder MenuBuilder(true, nullptr);
+
+				MenuBuilder.AddMenuEntry(LOCTEXT("MediaCaptureViewportMethod", "Viewport"),
+					LOCTEXT("MediaCaptureViewportMethodToolTip", "Capture from a viewport"),
+					FSlateIcon(),
+					FUIAction(FExecuteAction::CreateSP(this, &SMediaCaptureMethodComboButton::OnViewportMethodSelected)));
+
+				MenuBuilder.AddMenuEntry(LOCTEXT("MediaCaptureRenderTargetMethod", "Render Target"),
+					LOCTEXT("MediaCaptureRenderTargetMethodToolTip", "Capture from a render target"),
+					FSlateIcon(),
+					FUIAction(FExecuteAction::CreateSP(this, &SMediaCaptureMethodComboButton::OnRenderTargetMethodSelected)));
+				
+				return MenuBuilder.MakeWidget();
+			})
+			.ContentPadding(0.f)
+			.ButtonContent()
+			[
+				SNew(STextBlock)
+				.Text(this, &SMediaCaptureMethodComboButton::GetButtonText)
+				.Font(IDetailLayoutBuilder::GetDetailFont())
+			]);
+	}
+
+private:
+	void OnViewportMethodSelected()
+	{
+		UMediaFrameworkWorldSettingsAssetUserData* MediaCaptureSettings = FMediaProfileEditor::FindOrAddMediaFrameworkCaptureSettings();
+		if (!MediaCaptureSettings)
+		{
+			return;
+		}
+		
+		for (const TWeakObjectPtr<UObject>& WeakObject : Objects)
+		{
+			TStrongObjectPtr<UObject> Object = WeakObject.Pin();
+			UMediaOutput* MediaOutput = Object.IsValid() ? Cast<UMediaOutput>(Object.Get()) : nullptr;
+			if (!MediaOutput)
+			{
+				continue;
+			}
+
+			if (MediaOutput->IsA<UDummyMediaOutput>())
+			{
+				continue;
+			}
+			
+			// We expect there to be only one capture configuration per media output, so clear
+			// any existing viewport capture configurations before adding the new one
+			MediaCaptureSettings->ViewportCaptures.RemoveAll(GetFindOutputInfoPredicate<FMediaFrameworkCaptureCameraViewportCameraOutputInfo>(MediaOutput));
+
+			FMediaFrameworkCaptureCameraViewportCameraOutputInfo NewViewportOutputInfo;
+			NewViewportOutputInfo.MediaOutput = MediaOutput;
+			
+			// Now, if the object has an existing render target capture config, we want to copy that config over to the new viewport config, so find it
+			FMediaFrameworkCaptureRenderTargetCameraOutputInfo* ExistingRenderTargetOutputInfo =
+				MediaCaptureSettings->RenderTargetCaptures.FindByPredicate(GetFindOutputInfoPredicate<FMediaFrameworkCaptureRenderTargetCameraOutputInfo>(MediaOutput));
+			
+			if (ExistingRenderTargetOutputInfo)
+			{
+				NewViewportOutputInfo.CaptureOptions = ExistingRenderTargetOutputInfo->CaptureOptions;
+			}
+
+			MediaCaptureSettings->ViewportCaptures.Add(NewViewportOutputInfo);
+			MediaCaptureSettings->RenderTargetCaptures.RemoveAll(GetFindOutputInfoPredicate<FMediaFrameworkCaptureRenderTargetCameraOutputInfo>(MediaOutput));
+		}
+
+		OnMethodChanged.ExecuteIfBound();
+	}
+
+	void OnRenderTargetMethodSelected()
+	{
+		UMediaFrameworkWorldSettingsAssetUserData* MediaCaptureSettings = FMediaProfileEditor::FindOrAddMediaFrameworkCaptureSettings();
+		if (!MediaCaptureSettings)
+		{
+			return;
+		}
+		
+		for (const TWeakObjectPtr<UObject>& WeakObject : Objects)
+		{
+			TStrongObjectPtr<UObject> Object = WeakObject.Pin();
+			UMediaOutput* MediaOutput = Object.IsValid() ? Cast<UMediaOutput>(Object.Get()) : nullptr;
+			if (!MediaOutput)
+			{
+				continue;
+			}
+
+			if (MediaOutput->IsA<UDummyMediaOutput>())
+			{
+				continue;
+			}
+
+			// We expect there to be only one capture configuration per media output, so clear
+			// any existing render target capture configurations before adding the new one
+			MediaCaptureSettings->RenderTargetCaptures.RemoveAll(GetFindOutputInfoPredicate<FMediaFrameworkCaptureRenderTargetCameraOutputInfo>(MediaOutput));
+
+			FMediaFrameworkCaptureRenderTargetCameraOutputInfo NewRenderTargetOutputInfo;
+			NewRenderTargetOutputInfo.MediaOutput = MediaOutput;
+			
+			// Now, if the object has an existing render target capture config, we want to copy that config over to the new viewport config, so find it
+			FMediaFrameworkCaptureCameraViewportCameraOutputInfo* ExistingViewportOutputInfo =
+				MediaCaptureSettings->ViewportCaptures.FindByPredicate(GetFindOutputInfoPredicate<FMediaFrameworkCaptureCameraViewportCameraOutputInfo>(MediaOutput));
+			
+			if (ExistingViewportOutputInfo)
+			{
+				NewRenderTargetOutputInfo.CaptureOptions = ExistingViewportOutputInfo->CaptureOptions;
+			}
+
+			MediaCaptureSettings->RenderTargetCaptures.Add(NewRenderTargetOutputInfo);
+			MediaCaptureSettings->ViewportCaptures.RemoveAll(GetFindOutputInfoPredicate<FMediaFrameworkCaptureCameraViewportCameraOutputInfo>(MediaOutput));
+		}
+
+		OnMethodChanged.ExecuteIfBound();
+	}
+
+	FText GetButtonText() const
+	{
+		UMediaFrameworkWorldSettingsAssetUserData* MediaCaptureSettings = FMediaProfileEditor::FindOrAddMediaFrameworkCaptureSettings();
+		if (!MediaCaptureSettings)
+		{
+			return LOCTEXT("MediaCaptureMethodNotConfigured", "Not Configured");
+		}
+		
+		int32 NumViewports = 0;
+		int32 NumRenderTargets = 0;
+
+		for (const TWeakObjectPtr<UObject>& WeakObject : Objects)
+		{
+			TStrongObjectPtr<UObject> Object = WeakObject.Pin();
+			UMediaOutput* MediaOutput = Object.IsValid() ? Cast<UMediaOutput>(Object.Get()) : nullptr;
+			if (!MediaOutput)
+			{
+				continue;
+			}
+
+			NumViewports += Algo::CountIf(MediaCaptureSettings->ViewportCaptures, GetFindOutputInfoPredicate<FMediaFrameworkCaptureCameraViewportCameraOutputInfo>(MediaOutput));
+			NumRenderTargets += Algo::CountIf(MediaCaptureSettings->RenderTargetCaptures, GetFindOutputInfoPredicate<FMediaFrameworkCaptureRenderTargetCameraOutputInfo>(MediaOutput));
+		}
+
+		if (NumViewports > 0 && NumRenderTargets > 0)
+		{
+			return LOCTEXT("MediaCaptureMethodMultipleValues", "Multiple Values");
+		}
+		if (NumViewports > 0)
+		{
+			return LOCTEXT("MediaCaptureViewportMethod", "Viewport");
+		}
+		if (NumRenderTargets > 0)
+		{
+			return LOCTEXT("MediaCaptureRenderTargetMethod", "Render Target");
+		}
+
+		return LOCTEXT("MediaCaptureMethodNotConfigured", "Not Configured");
+	}
+	
+	template<typename TOutputInfo>
+	TFunction<bool(const TOutputInfo&)> GetFindOutputInfoPredicate(UMediaOutput* InMediaOutput) const
+	{
+		return [InMediaOutput](const TOutputInfo& OutputInfo)
+		{
+			return OutputInfo.MediaOutput == InMediaOutput;
+		};
+	}
+	
+private:
+	TArray<TWeakObjectPtr<UObject>> Objects;
+	FSimpleDelegate OnMethodChanged;
+};
+
+class FMediaFrameworkViewportOutputInfoCustomization : public IPropertyTypeCustomization
+{
+protected:
+	virtual void CustomizeHeader(TSharedRef<IPropertyHandle> PropertyHandle, FDetailWidgetRow& HeaderRow, IPropertyTypeCustomizationUtils& CustomizationUtils) override { }
+
+	virtual void CustomizeChildren(TSharedRef<IPropertyHandle> PropertyHandle, IDetailChildrenBuilder& ChildBuilder, IPropertyTypeCustomizationUtils& CustomizationUtils) override
+	{
+		ChildBuilder.AddProperty(PropertyHandle->GetChildHandle(GET_MEMBER_NAME_CHECKED(FMediaFrameworkCaptureCameraViewportCameraOutputInfo, LockedActors)).ToSharedRef());
+
+		// Manually add all capture options properties to avoid an additional property group getting added
+		TSharedPtr<IPropertyHandle> CaptureOptionsHandle = PropertyHandle->GetChildHandle(GET_MEMBER_NAME_CHECKED(FMediaFrameworkCaptureCameraViewportCameraOutputInfo, CaptureOptions));
+		uint32 NumChildren;
+		CaptureOptionsHandle->GetNumChildren(NumChildren);
+		for (uint32 Index = 0; Index < NumChildren; ++Index)
+		{
+			TSharedPtr<IPropertyHandle> ChildHandle = CaptureOptionsHandle->GetChildHandle(Index);
+			ChildBuilder.AddProperty(ChildHandle.ToSharedRef());
+		}
+	}
+};
+
+class FMediaFrameworkRenderTargetOutputInfoCustomization : public IPropertyTypeCustomization
+{
+protected:
+	virtual void CustomizeHeader(TSharedRef<IPropertyHandle> PropertyHandle, FDetailWidgetRow& HeaderRow, IPropertyTypeCustomizationUtils& CustomizationUtils) override { }
+
+	virtual void CustomizeChildren(TSharedRef<IPropertyHandle> PropertyHandle, IDetailChildrenBuilder& ChildBuilder, IPropertyTypeCustomizationUtils& CustomizationUtils) override
+	{
+		ChildBuilder.AddProperty(PropertyHandle->GetChildHandle(GET_MEMBER_NAME_CHECKED(FMediaFrameworkCaptureRenderTargetCameraOutputInfo, RenderTarget)).ToSharedRef());
+
+		// Manually add all capture options properties to avoid an additional property group getting added
+		TSharedPtr<IPropertyHandle> CaptureOptionsHandle = PropertyHandle->GetChildHandle(GET_MEMBER_NAME_CHECKED(FMediaFrameworkCaptureRenderTargetCameraOutputInfo, CaptureOptions));
+		uint32 NumChildren;
+		CaptureOptionsHandle->GetNumChildren(NumChildren);
+		for (uint32 Index = 0; Index < NumChildren; ++Index)
+		{
+			TSharedPtr<IPropertyHandle> ChildHandle = CaptureOptionsHandle->GetChildHandle(Index);
+			ChildBuilder.AddProperty(ChildHandle.ToSharedRef());
+		}
+	}
+};
+
+FMediaProfileMediaOutputDetailCustomization::~FMediaProfileMediaOutputDetailCustomization()
+{
+	if (MediaProfileEditor.IsValid())
+	{
+		MediaProfileEditor.Pin()->GetOnCaptureSettingsObjectChanged().RemoveAll(this);
+	}
+}
+
+void FMediaProfileMediaOutputDetailCustomization::CustomizeDetails(IDetailLayoutBuilder& DetailBuilder)
+{
+	FMediaProfileMediaItemDetailCustomization::CustomizeDetails(DetailBuilder);
+
+	MediaProfileEditor.Pin()->GetOnCaptureSettingsObjectChanged().AddSP(this, &FMediaProfileMediaOutputDetailCustomization::OnCaptureSettingsObjectChanged);
+	
+	TArray<TWeakObjectPtr<UObject>> Objects;
+	CachedDetailBuilder.Pin()->GetObjectsBeingCustomized(Objects);
+	
+	DetailBuilder.RegisterInstancedCustomPropertyTypeLayout(
+		FMediaFrameworkCaptureCameraViewportCameraOutputInfo::StaticStruct()->GetFName(),
+		FOnGetPropertyTypeCustomizationInstance::CreateLambda([]
+		{
+			return MakeShared<FMediaFrameworkViewportOutputInfoCustomization>();
+		}));
+
+	DetailBuilder.RegisterInstancedCustomPropertyTypeLayout(
+		FMediaFrameworkCaptureRenderTargetCameraOutputInfo::StaticStruct()->GetFName(),
+		FOnGetPropertyTypeCustomizationInstance::CreateLambda([]
+		{
+			return MakeShared<FMediaFrameworkRenderTargetOutputInfoCustomization>();
+		}));
+	
+	// For media outputs, add the corresponding capture info to the details panel
+	IDetailCategoryBuilder& MediaCaptureCategory = DetailBuilder.EditCategory(
+		TEXT("MediaCaptureCategory"),
+		LOCTEXT("MediaCaptureCategoryLabel", "Capture"),
+		ECategoryPriority::Important);
+
+	MediaCaptureCategory
+		.AddCustomRow(LOCTEXT("MediaCaptureButtonFilter", "Start/Stop Media Capture"))
+		.Visibility(TAttribute<EVisibility>::CreateSP(this, &FMediaProfileMediaOutputDetailCustomization::GetValidObjectVisibility))
+		.WholeRowContent()
+		[
+			SNew(SButton)
+			.HAlign(HAlign_Center)
+			.VAlign(VAlign_Center)
+			.OnPressed(this, &FMediaProfileMediaOutputDetailCustomization::OnCaptureButtonPressed)
+			[
+				SNew(SHorizontalBox)
+
+				+SHorizontalBox::Slot()
+				.AutoWidth()
+				.HAlign(HAlign_Center)
+				.VAlign(VAlign_Center)
+				.Padding(0.0f, 0.0f, 4.0f, 0.0f)				
+				[
+					SNew(SImage)
+					.Image(this, &FMediaProfileMediaOutputDetailCustomization::GetCaptureButtonImage)
+				]
+
+				+SHorizontalBox::Slot()
+				.AutoWidth()
+				.HAlign(HAlign_Center)
+				.VAlign(VAlign_Center)
+				[
+					SNew(STextBlock)
+					.Text(this, &FMediaProfileMediaOutputDetailCustomization::GetCaptureButtonText)
+				]
+			]
+			.IsEnabled(this, &FMediaProfileMediaOutputDetailCustomization::IsCaptureButtonEnabled)
+		];
+		
+	MediaCaptureCategory
+		.AddCustomRow(LOCTEXT("MediaCaptureMethodFilter", "Capture Method"))
+        .Visibility(TAttribute<EVisibility>::CreateSP(this, &FMediaProfileMediaOutputDetailCustomization::GetValidObjectVisibility))
+        .NameContent()
+		[
+			SNew(STextBlock)
+			.Text(LOCTEXT("MediaCaptureMethodLabel", "Capture Method"))
+			.Font(IDetailLayoutBuilder::GetDetailFont())
+		]
+		.ValueContent()
+		.MinDesiredWidth(125.f)
+		.MaxDesiredWidth(250.f)
+		[
+			SNew(SHorizontalBox)
+			+ SHorizontalBox::Slot()
+			.FillWidth(1.0f)
+			[
+				SNew(SMediaCaptureMethodComboButton, Objects)
+				.OnMethodChanged_Lambda([this]()
+				{
+					CachedDetailBuilder.Pin()->ForceRefreshDetails();
+				})
+			]
+		];
+
+	TArray<TSharedPtr<FStructOnScope>> CaptureSettingStructs;
+	if (UMediaFrameworkWorldSettingsAssetUserData* MediaCaptureSettings = FMediaProfileEditor::FindOrAddMediaFrameworkCaptureSettings())
+	{
+		for (const TWeakObjectPtr<UObject>& WeakObject : Objects)
+		{
+			TStrongObjectPtr<UObject> Object = WeakObject.Pin();			
+			UMediaOutput* MediaOutput = Object.IsValid() ? Cast<UMediaOutput>(Object.Get()) : nullptr;
+			if (!MediaOutput)
+			{
+				continue;
+			}
+
+			if (MediaOutput->IsA<UDummyMediaOutput>())
+			{
+				continue;
+			}
+
+			for (FMediaFrameworkCaptureCameraViewportCameraOutputInfo& ViewportOutputInfo : MediaCaptureSettings->ViewportCaptures)
+			{
+				if (ViewportOutputInfo.MediaOutput == MediaOutput)
+				{
+					CaptureSettingStructs.Add(MakeShared<FStructOnScope>(FMediaFrameworkCaptureCameraViewportCameraOutputInfo::StaticStruct(), reinterpret_cast<uint8*>(&ViewportOutputInfo)));
+				}
+			}
+
+			for (FMediaFrameworkCaptureRenderTargetCameraOutputInfo& RenderTargetOutputInfo : MediaCaptureSettings->RenderTargetCaptures)
+			{
+				if (RenderTargetOutputInfo.MediaOutput == MediaOutput)
+				{
+					CaptureSettingStructs.Add(MakeShared<FStructOnScope>(FMediaFrameworkCaptureRenderTargetCameraOutputInfo::StaticStruct(), reinterpret_cast<uint8*>(&RenderTargetOutputInfo)));
+				}
+			}
+		}
+	}
+
+	if (CaptureSettingStructs.Num())
+	{
+		FAddPropertyParams AddPropertyParams;
+		AddPropertyParams.HideRootObjectNode(true);
+			
+		MediaCaptureStruct = MakeShared<FStructOnScopeStructureDataProvider>(CaptureSettingStructs);
+		MediaCaptureCategory.AddExternalStructureProperty(MediaCaptureStruct, NAME_None, EPropertyLocation::Default, AddPropertyParams);
+	}
+}
+
+FText FMediaProfileMediaOutputDetailCustomization::GetMediaItemTypeText() const
+{
+	return LOCTEXT("MediaOutputTypeText", "Media Output");
+}
+
+FText FMediaProfileMediaOutputDetailCustomization::GetMediaObjectLabel(UMediaOutput* InMediaObject) const
+{
+	if (!MediaProfile.IsValid())
+	{
+		return FText::GetEmpty();
+	}
+
+	int32 Index = MediaProfile->FindMediaOutputIndex(InMediaObject);
+	if (InMediaObject->IsA<UDummyMediaOutput>())
+	{
+		Index = Cast<UDummyMediaOutput>(InMediaObject)->MediaProfileIndex;
+	}
+		
+	return FText::FromString(MediaProfile->GetLabelForMediaOutput(Index));
+}
+
+void FMediaProfileMediaOutputDetailCustomization::SetMediaObjectLabel(UMediaOutput* InMediaObject, const FText& InLabel)
+{
+	if (!MediaProfile.IsValid())
+	{
+		return;
+	}
+	
+	int32 Index = MediaProfile->FindMediaOutputIndex(InMediaObject);
+	if (InMediaObject->IsA<UDummyMediaOutput>())
+	{
+		Index = Cast<UDummyMediaOutput>(InMediaObject)->MediaProfileIndex;
+	}
+		
+	MediaProfile->SetLabelForMediaOutput(Index, InLabel.ToString());
+}
+
+void FMediaProfileMediaOutputDetailCustomization::SetMediaObject(UMediaOutput* InOriginalMediaObject, UMediaOutput* InNewMediaObject)
+{
+	if (!MediaProfile.IsValid())
+	{
+		return;
+	}
+	
+	int32 Index = MediaProfile->FindMediaOutputIndex(InOriginalMediaObject);
+	if (InOriginalMediaObject->IsA<UDummyMediaOutput>())
+	{
+		Index = Cast<UDummyMediaOutput>(InOriginalMediaObject)->MediaProfileIndex;
+	}
+		
+	MediaProfile->SetMediaOutput(Index, InNewMediaObject);
+	
+	UMediaFrameworkWorldSettingsAssetUserData* CaptureSettings = FMediaProfileEditor::FindOrAddMediaFrameworkCaptureSettings();
+	if (!CaptureSettings)
+	{
+		return;
+	}
+
+	bool bHasExistingCaptureInfo = false;
+	for (FMediaFrameworkCaptureCameraViewportCameraOutputInfo& ViewportOutputInfo : CaptureSettings->ViewportCaptures)
+	{
+		if (ViewportOutputInfo.MediaOutput != InOriginalMediaObject)
+		{
+			continue;
+		}
+
+		ViewportOutputInfo.MediaOutput = InNewMediaObject;
+		bHasExistingCaptureInfo = true;
+	}
+
+	for (FMediaFrameworkCaptureRenderTargetCameraOutputInfo& RenderTargetOutputInfo : CaptureSettings->RenderTargetCaptures)
+	{
+		if (RenderTargetOutputInfo.MediaOutput != InOriginalMediaObject)
+		{
+			continue;
+		}
+
+		RenderTargetOutputInfo.MediaOutput = InNewMediaObject;
+		bHasExistingCaptureInfo = true;
+	}
+
+	if (!bHasExistingCaptureInfo)
+	{
+		FMediaFrameworkCaptureCameraViewportCameraOutputInfo NewCameraOutputInfo;
+		NewCameraOutputInfo.MediaOutput = InNewMediaObject;
+		CaptureSettings->ViewportCaptures.Add(NewCameraOutputInfo);
+	}
+}
+
+void FMediaProfileMediaOutputDetailCustomization::OnCaptureButtonPressed()
+{
+	if (MediaProfileEditor.IsValid())
+	{
+		TSharedPtr<FMediaProfileEditor> PinnedEditor = MediaProfileEditor.Pin();
+			
+		TArray<TWeakObjectPtr<UObject>> Objects;
+		CachedDetailBuilder.Pin()->GetObjectsBeingCustomized(Objects);
+
+		bool bIsCapturing = false;
+		for (const TWeakObjectPtr<UObject>& WeakObject : Objects)
+		{
+			TStrongObjectPtr<UObject> Object = WeakObject.Pin();
+			if (Object.IsValid() && Object.Get()->IsA<UMediaOutput>())
+			{
+				if (PinnedEditor->IsMediaOutputCapturing(Cast<UMediaOutput>(Object.Get())))
+				{
+					bIsCapturing = true;
+					break;
+				}
+			}
+		}
+			
+		if (bIsCapturing)
+		{
+			for (const TWeakObjectPtr<UObject>& WeakObject : Objects)
+			{
+				TStrongObjectPtr<UObject> Object = WeakObject.Pin();
+				if (Object.IsValid() && Object.Get()->IsA<UMediaOutput>())
+				{
+					PinnedEditor->StopCapturingMediaOutput(Cast<UMediaOutput>(Object.Get()));
+				}
+			}
+		}
+		else
+		{
+			for (const TWeakObjectPtr<UObject>& WeakObject : Objects)
+			{
+			TStrongObjectPtr<UObject> Object = WeakObject.Pin();
+				if (Object.IsValid() && Object.Get()->IsA<UMediaOutput>())
+				{
+					PinnedEditor->StartCapturingMediaOutput(Cast<UMediaOutput>(Object.Get()));
+				}
+			}
+		}
+	}
+}
+
+const FSlateBrush* FMediaProfileMediaOutputDetailCustomization::GetCaptureButtonImage() const
+{
+	bool bIsCapturing = false;
+	if (MediaProfileEditor.IsValid())
+	{
+		TSharedPtr<FMediaProfileEditor> PinnedEditor = MediaProfileEditor.Pin();
+			
+		TArray<TWeakObjectPtr<UObject>> Objects;
+		CachedDetailBuilder.Pin()->GetObjectsBeingCustomized(Objects);
+
+		for (const TWeakObjectPtr<UObject>& WeakObject : Objects)
+		{
+			TStrongObjectPtr<UObject> Object = WeakObject.Pin();
+			if (Object.IsValid() && Object.Get()->IsA<UMediaOutput>())
+			{
+				if (PinnedEditor->IsMediaOutputCapturing(Cast<UMediaOutput>(Object.Get())))
+				{
+					bIsCapturing = true;
+					break;
+				}
+			}
+		}
+	}
+		
+	return !bIsCapturing ? FMediaFrameworkUtilitiesEditorStyle::Get().GetBrush("MediaCapture.Capture") : FAppStyle::GetBrush("Icons.Toolbar.Stop");
+}
+
+FText FMediaProfileMediaOutputDetailCustomization::GetCaptureButtonText() const
+{
+	bool bIsCapturing = false;
+	if (MediaProfileEditor.IsValid())
+	{
+		TSharedPtr<FMediaProfileEditor> PinnedEditor = MediaProfileEditor.Pin();
+			
+		TArray<TWeakObjectPtr<UObject>> Objects;
+		CachedDetailBuilder.Pin()->GetObjectsBeingCustomized(Objects);
+
+		for (const TWeakObjectPtr<UObject>& WeakObject : Objects)
+		{
+			TStrongObjectPtr<UObject> Object = WeakObject.Pin();
+			if (Object.IsValid() && Object.Get()->IsA<UMediaOutput>())
+			{
+				if (PinnedEditor->IsMediaOutputCapturing(Cast<UMediaOutput>(Object.Get())))
+				{
+					bIsCapturing = true;
+					break;
+				}
+			}
+		}
+	}
+		
+	return !bIsCapturing ? LOCTEXT("MediaOutputStartCapture", "Start Capture") : LOCTEXT("MediaOutputStopCapture", "Stop Capture");
+}
+
+bool FMediaProfileMediaOutputDetailCustomization::IsCaptureButtonEnabled() const
+{
+	bool bIsCapturing = false;
+	bool bIsInvalidCaptureConfig = false;
+	if (MediaProfileEditor.IsValid())
+	{
+		TSharedPtr<FMediaProfileEditor> PinnedEditor = MediaProfileEditor.Pin();
+			
+		TArray<TWeakObjectPtr<UObject>> Objects;
+		CachedDetailBuilder.Pin()->GetObjectsBeingCustomized(Objects);
+
+		for (const TWeakObjectPtr<UObject>& WeakObject : Objects)
+		{
+			TStrongObjectPtr<UObject> Object = WeakObject.Pin();
+			if (Object.IsValid() && Object.Get()->IsA<UMediaOutput>())
+			{
+				UMediaOutput* MediaOutput = Cast<UMediaOutput>(Object.Get());
+				if (PinnedEditor->IsMediaOutputCapturing(MediaOutput))
+				{
+					bIsCapturing = true;
+				}
+					
+				if (!PinnedEditor->CanMediaOutputCapture(MediaOutput))
+				{
+					bIsInvalidCaptureConfig = true;
+				}
+			}
+		}
+	}
+		
+	return bIsCapturing || !bIsInvalidCaptureConfig;
+}
+
+void FMediaProfileMediaOutputDetailCustomization::OnCaptureSettingsObjectChanged()
+{
+	if (CachedDetailBuilder.IsValid())
+	{
+		CachedDetailBuilder.Pin()->ForceRefreshDetails();
+	}
+}
+
+#undef LOCTEXT_NAMESPACE
diff --git a/Engine/Plugins/Media/MediaFrameworkUtilities/Source/MediaFrameworkUtilitiesEditor/Private/AssetEditor/MediaProfileMediaItemDetailCustomization.h b/Engine/Plugins/Media/MediaFrameworkUtilities/Source/MediaFrameworkUtilitiesEditor/Private/AssetEditor/MediaProfileMediaItemDetailCustomization.h
new file mode 100644
--- /dev/null
+++ b/Engine/Plugins/Media/MediaFrameworkUtilities/Source/MediaFrameworkUtilitiesEditor/Private/AssetEditor/MediaProfileMediaItemDetailCustomization.h
@@ -0,0 +1,349 @@
+﻿// Copyright Epic Games, Inc. All Rights Reserved.
+
+#pragma once
+
+#include "CoreMinimal.h"
+#include "ClassViewerFilter.h"
+#include "ClassViewerModule.h"
+#include "DetailCategoryBuilder.h"
+#include "DetailLayoutBuilder.h"
+#include "DetailWidgetRow.h"
+#include "DummyMediaObject.h"
+#include "IDetailCustomization.h"
+#include "MediaProfileEditor.h"
+#include "ScopedTransaction.h"
+#include "Modules/ModuleManager.h"
+#include "Profile/MediaProfile.h"
+#include "Styling/SlateIconFinder.h"
+#include "Widgets/Images/SImage.h"
+#include "Widgets/Input/SComboButton.h"
+#include "Widgets/Input/SEditableTextBox.h"
+
+#define LOCTEXT_NAMESPACE "FMediaProfileMediaItemDetailCustomization"
+
+class FStructOnScopeStructureDataProvider;
+
+/**
+ * Base class for a details panel customization for a media source or output within a media profile. Displays an additional category that
+ * contains a editable label field and a type dropdown that allows the user to change the media item's type
+ */
+template<typename TMediaType>
+class FMediaProfileMediaItemDetailCustomization : public IDetailCustomization
+{
+private:
+	class FMediaTypeClassFilter : public IClassViewerFilter
+	{
+	public:
+		virtual bool IsClassAllowed(const FClassViewerInitializationOptions& InInitOptions, const UClass* InClass, TSharedRef<FClassViewerFilterFuncs> InFilterFuncs) override
+		{
+			if (InClass != nullptr)
+			{
+				return InClass->IsChildOf(TMediaType::StaticClass()) &&
+					!InClass->HasAnyClassFlags(CLASS_Abstract) &&
+					!(InClass == UDummyMediaSource::StaticClass() || InClass == UDummyMediaOutput::StaticClass());
+			}
+			return false;
+		}
+
+		virtual bool IsUnloadedClassAllowed(const FClassViewerInitializationOptions& InInitOptions, const TSharedRef<const IUnloadedBlueprintData> InUnloadedClassData, TSharedRef<FClassViewerFilterFuncs> InFilterFuncs) override
+		{
+			return InUnloadedClassData->IsChildOf(TMediaType::StaticClass()) && !InUnloadedClassData->HasAnyClassFlags(CLASS_Abstract);
+		}
+	};
+	
+public:
+	FMediaProfileMediaItemDetailCustomization(TWeakPtr<FMediaProfileEditor>& InMediaProfileEditor, UMediaProfile* InMediaProfile, FSimpleDelegate InOnObjectsChanged)
+	{
+		MediaProfileEditor = InMediaProfileEditor;
+		MediaProfile = InMediaProfile;
+		OnObjectsChanged = InOnObjectsChanged;
+	}
+
+	virtual void CustomizeDetails(const TSharedPtr<IDetailLayoutBuilder>& DetailBuilder) override
+	{
+		CachedDetailBuilder = DetailBuilder;
+		IDetailCustomization::CustomizeDetails(DetailBuilder);
+	}
+
+	virtual void CustomizeDetails(IDetailLayoutBuilder& DetailBuilder) override
+	{
+		const FText CategoryLabel = FText::Format(LOCTEXT("MediaItemCategoryName", "{0}"), GetMediaItemTypeText());
+		IDetailCategoryBuilder& MediaTypeCategory = DetailBuilder.EditCategory(TEXT("MediaTypeCategory"), CategoryLabel, ECategoryPriority::Important);
+
+		MediaTypeCategory.AddCustomRow(LOCTEXT("LabelFilterText", "Label"))
+			.Visibility(TAttribute<EVisibility>::CreateSP(this, &FMediaProfileMediaItemDetailCustomization::GetValidObjectVisibility))
+			.NameContent()
+			[
+				SNew(STextBlock)
+				.Text(LOCTEXT("LabelName", "Label"))
+				.Font(IDetailLayoutBuilder::GetDetailFont())
+			]
+			.ValueContent()
+			[
+				SNew(SEditableTextBox)
+				.Font(IDetailLayoutBuilder::GetDetailFont())
+				.Text(this, &FMediaProfileMediaItemDetailCustomization::GetLabelText)
+				.IsEnabled(this, &FMediaProfileMediaItemDetailCustomization::IsLabelTextBoxEnabled)
+				.OnTextCommitted(this, &FMediaProfileMediaItemDetailCustomization::OnLabelTextCommitted)
+			];
+
+		const FText TypeFilterText = FText::Format(LOCTEXT("MediaItemTypeFilterText", "{0} Type"), GetMediaItemTypeText());
+		const FText TypeLabel = FText::Format(LOCTEXT("MediaItemTypeName", "{0}"), GetMediaItemTypeText());
+		
+		MediaTypeCategory.AddCustomRow(TypeFilterText)
+			.NameContent()
+			[
+				SNew(STextBlock)
+				.Text(TypeLabel)
+				.Font(IDetailLayoutBuilder::GetDetailFont())
+			]
+			.ValueContent()
+			.MinDesiredWidth(125.f)
+			.MaxDesiredWidth(250.f)
+			[
+				SNew(SHorizontalBox)
+				+ SHorizontalBox::Slot()
+				.FillWidth(1.0f)
+				[
+					SAssignNew(ClassPickerComboButton, SComboButton)
+					.OnGetMenuContent(this, &FMediaProfileMediaItemDetailCustomization::GetClassPickerMenuContent)
+					.ContentPadding(0.f)
+					.ButtonContent()
+					[
+						SNew(SHorizontalBox)
+						+SHorizontalBox::Slot()
+						.AutoWidth()
+						.VAlign(VAlign_Center)
+						.Padding(0.0f, 0.0f, 4.0f, 0.0f)
+						[
+							SNew(SImage)
+							.Image(this, &FMediaProfileMediaItemDetailCustomization::GetClassIcon)
+						]
+						+SHorizontalBox::Slot()
+						.VAlign(VAlign_Center)
+						[
+							SNew(STextBlock )
+							.Text(this, &FMediaProfileMediaItemDetailCustomization::GetClassName)
+							.Font(IDetailLayoutBuilder::GetDetailFont())
+						]
+					]
+				]
+			];
+	}
+
+protected:
+	virtual FText GetMediaItemTypeText() const = 0;
+	
+	EVisibility GetValidObjectVisibility() const
+	{
+		TArray<TWeakObjectPtr<UObject>> Objects;
+		CachedDetailBuilder.Pin()->GetObjectsBeingCustomized(Objects);
+
+		bool bContainsValidObjects = Objects.ContainsByPredicate([](const TWeakObjectPtr<UObject>& Object)
+		{
+			return Object.IsValid() && !(Object->IsA<UDummyMediaSource>() || Object->IsA<UDummyMediaOutput>());
+		});
+
+		return bContainsValidObjects ? EVisibility::Visible : EVisibility::Hidden;
+	}
+	
+	FText GetLabelText() const
+	{
+		TArray<TWeakObjectPtr<UObject>> Objects;
+		CachedDetailBuilder.Pin()->GetObjectsBeingCustomized(Objects);
+
+		if (Objects.Num() <= 0)
+		{
+			return FText::GetEmpty();
+		}
+		
+		if (Objects.Num() > 1)
+		{
+			return LOCTEXT("MultipleValuesText", "Multiple Values");
+		}
+
+		TMediaType* MediaObject = Cast<TMediaType>(Objects[0]);
+		return GetMediaObjectLabel(MediaObject);
+	}
+	
+	bool IsLabelTextBoxEnabled() const
+	{
+		TArray<TWeakObjectPtr<UObject>> Objects;
+		CachedDetailBuilder.Pin()->GetObjectsBeingCustomized(Objects);
+
+		return Objects.Num() == 1;
+	}
+	
+	void OnLabelTextCommitted(const FText& NewText, ETextCommit::Type CommitInfo)
+	{
+		TArray<TWeakObjectPtr<UObject>> Objects;
+		CachedDetailBuilder.Pin()->GetObjectsBeingCustomized(Objects);
+
+		if (Objects.Num() != 1)
+		{
+			return;
+		}
+		
+		TMediaType* MediaObject = Cast<TMediaType>(Objects[0]);
+		
+		FScopedTransaction Transaction(LOCTEXT("SetLabelTransaction", "Set Label"));
+		MediaProfile->Modify();
+
+		SetMediaObjectLabel(MediaObject, NewText);
+	}
+
+	TSharedRef<SWidget> GetClassPickerMenuContent()
+	{
+		TSharedPtr<FMediaTypeClassFilter> Filter = MakeShared<FMediaTypeClassFilter>();
+		
+		FClassViewerInitializationOptions Options;
+		Options.bShowBackgroundBorder = false;
+		Options.bShowUnloadedBlueprints = true;
+		Options.NameTypeToDisplay = EClassViewerNameTypeToDisplay::DisplayName;
+		Options.ClassFilters.Add(Filter.ToSharedRef());
+		
+		return FModuleManager::LoadModuleChecked<FClassViewerModule>("ClassViewer").CreateClassViewer(Options, FOnClassPicked::CreateSP(this, &FMediaProfileMediaItemDetailCustomization::OnClassPicked));
+	}
+	
+	void OnClassPicked(UClass* InClass)
+	{
+		ClassPickerComboButton->SetIsOpen(false);
+
+		TArray<TWeakObjectPtr<UObject>> Objects;
+		CachedDetailBuilder.Pin()->GetObjectsBeingCustomized(Objects);
+
+		if (Objects.Num() <= 0)
+		{
+			return;
+		}
+		
+		FScopedTransaction Transaction(FText::Format(LOCTEXT("SetMediaItemTypeTransaction", "Set {0}"), GetMediaItemTypeText()));
+		MediaProfile->Modify();
+		
+		for (int32 Index = 0; Index < Objects.Num(); ++Index)
+		{
+			TMediaType* OriginalMediaObject = Cast<TMediaType>(Objects[Index]);
+			TMediaType* NewMediaObject = NewObject<TMediaType>(MediaProfile.Get(), InClass);
+			
+			SetMediaObject(OriginalMediaObject, NewMediaObject);
+		}
+
+		OnObjectsChanged.ExecuteIfBound();
+	}
+	
+	const FSlateBrush* GetClassIcon() const
+	{
+		TArray<TWeakObjectPtr<UObject>> Objects;
+		CachedDetailBuilder.Pin()->GetObjectsBeingCustomized(Objects);
+
+		if (Objects.Num() != 1)
+		{
+			return nullptr;
+		}
+
+		if (Objects[0]->IsA<UDummyMediaSource>() || Objects[0]->IsA<UDummyMediaOutput>())
+		{
+			return nullptr;
+		}
+		
+		return FSlateIconFinder::FindIconForClass(Objects[0]->GetClass()).GetIcon();
+	}
+	
+	FText GetClassName() const
+	{
+		TArray<TWeakObjectPtr<UObject>> Objects;
+		CachedDetailBuilder.Pin()->GetObjectsBeingCustomized(Objects);
+
+		if (Objects.Num() <= 0)
+		{
+			return FText::GetEmpty();
+		}
+		
+		if (Objects.Num() > 1)
+		{
+			return LOCTEXT("MultipleValuesText", "Multiple Values");
+		}
+
+		if (Objects[0]->IsA<UDummyMediaSource>() || Objects[0]->IsA<UDummyMediaOutput>())
+		{
+			return LOCTEXT("NotSetText", "Not Set");
+		}
+		
+		return Objects[0]->GetClass()->GetDisplayNameText();
+	}
+	
+	virtual FText GetMediaObjectLabel(TMediaType* InMediaObject) const = 0;
+	virtual void SetMediaObjectLabel(TMediaType* InMediaObject, const FText& InLabel) = 0;
+	virtual void SetMediaObject(TMediaType* InOriginalMediaObject, TMediaType* InNewMediaObject) = 0;
+	
+protected:
+	/** Media profile editor that owns the details panel being customized by this customization */
+	TWeakPtr<FMediaProfileEditor> MediaProfileEditor;
+
+	/** Media profile that owns the media sources or outputs whose details panel is being customized */
+	TWeakObjectPtr<UMediaProfile> MediaProfile;
+
+	/** Cached detail layout builder pointer, used to manually refresh the details panel */
+	TWeakPtr<IDetailLayoutBuilder> CachedDetailBuilder;
+
+	/** Combo button that displays the types that the media item can be changed to */
+	TSharedPtr<SComboButton> ClassPickerComboButton;
+
+	/** Delegate that is raised when the media items' type is changed */
+	FSimpleDelegate OnObjectsChanged;
+};
+
+/**
+ * Details panel customization for media sources
+ */
+class FMediaProfileMediaSourceDetailCustomization : public FMediaProfileMediaItemDetailCustomization<UMediaSource>
+{
+public:
+	FMediaProfileMediaSourceDetailCustomization(TWeakPtr<FMediaProfileEditor>& InMediaProfileEditor, UMediaProfile* InMediaProfile, FSimpleDelegate InOnObjectsChanged)
+		: FMediaProfileMediaItemDetailCustomization(InMediaProfileEditor, InMediaProfile, InOnObjectsChanged)
+	{ }
+
+private:
+	virtual FText GetMediaItemTypeText() const override;
+	virtual FText GetMediaObjectLabel(UMediaSource* InMediaObject) const override;
+	virtual void SetMediaObjectLabel(UMediaSource* InMediaObject, const FText& InLabel) override;
+	virtual void SetMediaObject(UMediaSource* InOriginalMediaObject, UMediaSource* InNewMediaObject) override;
+};
+
+/**
+ * Details customization for media outputs. Adds a media capture category that displays the outputs' media capture settings
+ */
+class FMediaProfileMediaOutputDetailCustomization : public FMediaProfileMediaItemDetailCustomization<UMediaOutput>
+{
+public:
+	FMediaProfileMediaOutputDetailCustomization(TWeakPtr<FMediaProfileEditor>& InMediaProfileEditor, UMediaProfile* InMediaProfile, FSimpleDelegate InOnObjectsChanged)
+		: FMediaProfileMediaItemDetailCustomization(InMediaProfileEditor, InMediaProfile, InOnObjectsChanged)
+	{ }
+
+	virtual ~FMediaProfileMediaOutputDetailCustomization() override;
+	
+	virtual void CustomizeDetails(IDetailLayoutBuilder& DetailBuilder) override;
+	
+private:
+	virtual FText GetMediaItemTypeText() const override;
+	virtual FText GetMediaObjectLabel(UMediaOutput* InMediaObject) const override;
+	virtual void SetMediaObjectLabel(UMediaOutput* InMediaObject, const FText& InLabel) override;
+	virtual void SetMediaObject(UMediaOutput* InOriginalMediaObject, UMediaOutput* InNewMediaObject) override;
+	
+	void OnCaptureButtonPressed();
+	
+	const FSlateBrush* GetCaptureButtonImage() const;
+	FText GetCaptureButtonText() const;
+	bool IsCaptureButtonEnabled() const;
+	
+	void OnCaptureSettingsObjectChanged();
+
+private:
+	/**
+	 * Stored FStructOnScopes for the media capture structs of all objects being displayed by this customization, which are pulled from
+	 * the currently active UMediaFrameworkWorldSettingsAssetUserData's ViewportCaptures and RenderTargetCaptures list
+	 */
+	TSharedPtr<FStructOnScopeStructureDataProvider> MediaCaptureStruct;
+};
+
+#undef LOCTEXT_NAMESPACE
diff --git a/Engine/Plugins/Media/MediaFrameworkUtilities/Source/MediaFrameworkUtilitiesEditor/Private/AssetEditor/SMediaProfileDetailsInfoPanel.cpp b/Engine/Plugins/Media/MediaFrameworkUtilities/Source/MediaFrameworkUtilitiesEditor/Private/AssetEditor/SMediaProfileDetailsInfoPanel.cpp
new file mode 100644
--- /dev/null
+++ b/Engine/Plugins/Media/MediaFrameworkUtilities/Source/MediaFrameworkUtilitiesEditor/Private/AssetEditor/SMediaProfileDetailsInfoPanel.cpp
@@ -0,0 +1,544 @@
+﻿// Copyright Epic Games, Inc. All Rights Reserved.
+
+#include "SMediaProfileDetailsInfoPanel.h"
+
+#include "DetailLayoutBuilder.h"
+#include "DummyMediaObject.h"
+#include "MediaFrameworkWorldSettingsAssetUserData.h"
+#include "MediaOutput.h"
+#include "MediaProfileEditor.h"
+#include "MediaSource.h"
+#include "Engine/TextureRenderTarget2D.h"
+#include "MediaAssets/ProxyMediaOutput.h"
+#include "MediaAssets/ProxyMediaSource.h"
+#include "Profile/MediaProfile.h"
+#include "Styling/SlateIconFinder.h"
+#include "Widgets/SBoxPanel.h"
+#include "Widgets/Images/SImage.h"
+#include "Widgets/Layout/SBox.h"
+#include "Widgets/Text/STextBlock.h"
+
+#define LOCTEXT_NAMESPACE "SMediaObjectInfoPanel"
+
+TSharedRef<SWidget> SMediaObjectInfoPanel::FTwoColumnInfo::CreateWidget()
+{
+	return SNew(SHorizontalBox)
+		+SHorizontalBox::Slot()
+		.AutoWidth()
+		.Padding(0.0f, 0.0f, 12.0f, 2.0f)
+		[
+			SAssignNew(LabelColumn, SVerticalBox)
+		]
+
+		+SHorizontalBox::Slot()
+		.AutoWidth()
+		[
+			SAssignNew(ValueColumn, SVerticalBox)
+		];
+}
+
+void SMediaObjectInfoPanel::FTwoColumnInfo::AddEntry(TAttribute<FText> Label, TAttribute<FText> Value, TAttribute<FSlateColor> ValueColor)
+{
+	LabelColumn->AddSlot()
+	    .Padding(0.0, 2.0f)
+		[
+			SNew(STextBlock)
+			.Text(Label)
+			.Font(IDetailLayoutBuilder::GetDetailFont())
+		];
+
+	ValueColumn->AddSlot()
+	    .Padding(0.0, 2.0f)
+		[
+			SNew(STextBlock)
+			.Text(Value)
+			.Font(IDetailLayoutBuilder::GetDetailFont())
+			.ColorAndOpacity(ValueColor)
+		];
+}
+
+void SMediaObjectInfoPanel::FTwoColumnInfo::AddEntry(TAttribute<FText> Value, const FSlateBrush* Icon)
+{
+	LabelColumn->AddSlot()
+		.Padding(0.0, 2.0f)
+		[
+			SNullWidget::NullWidget
+		];
+
+	ValueColumn->AddSlot()
+		.Padding(0.0, 2.0f)
+		[
+			SNew(SHorizontalBox)
+			+SHorizontalBox::Slot()
+			.HAlign(HAlign_Center)
+			.VAlign(VAlign_Center)
+			.AutoWidth()
+			[
+				SNew(SBox)
+				.WidthOverride(16.0f)
+				.HeightOverride(16.0f)
+				[
+					SNew(SImage).Image(Icon)
+				]
+			]
+
+			+SHorizontalBox::Slot()
+			.FillWidth(1.0f)
+			.Padding(6.0f, 0.0f, 0.0f, 0.0f)
+			[
+				SNew(STextBlock)
+				.Text(Value)
+				.Font(IDetailLayoutBuilder::GetDetailFont())
+			]
+		];
+}
+
+void SMediaObjectInfoPanel::FTwoColumnInfo::AddEntry(TAttribute<FText> Value)
+{
+	LabelColumn->AddSlot()
+		.Padding(0.0, 2.0f)
+		[
+			SNullWidget::NullWidget
+		];
+
+	ValueColumn->AddSlot()
+		.Padding(0.0, 2.0f)
+		[
+			SNew(STextBlock)
+			.Text(Value)
+			.Font(IDetailLayoutBuilder::GetDetailFont())
+		];
+}
+
+void SMediaObjectInfoPanel::FTwoColumnInfo::ClearEntries()
+{
+	if (LabelColumn.IsValid())
+	{
+		LabelColumn->ClearChildren();
+	}
+
+	if (ValueColumn.IsValid())
+	{
+		ValueColumn->ClearChildren();
+	}
+}
+
+void SMediaObjectInfoPanel::Construct(const FArguments& InArgs, TWeakPtr<FMediaProfileEditor>& InMediaProfileEditor, UMediaProfile* InMediaProfile)
+{
+	MediaProfileEditor = InMediaProfileEditor;
+	MediaProfile = InMediaProfile;
+		
+	ChildSlot
+	[
+		SNew(SVerticalBox)
+		+SVerticalBox::Slot()
+		.AutoHeight()
+		.Padding(0.0f, 0.0f, 0.0f, 16.0f)
+		[
+			SNew(SHorizontalBox)
+			+SHorizontalBox::Slot()
+			.FillWidth(0.5f)
+			[
+				CreateInfoGroupWidget(LOCTEXT("BasicInfoHeader", "Info"), BasicInfo)
+			]
+
+			+SHorizontalBox::Slot()
+			.FillWidth(0.5f)
+			[
+				CreateInfoGroupWidget(TAttribute<FText>::CreateLambda([this]()
+				{
+					return ObjectAsMediaSource() ? LOCTEXT("MediaSourceInfoHeader", "Media Source") : LOCTEXT("MediaOutputInfoHeader", "Media Output");
+				}), TypeInfo)
+			]
+		]
+			
+		+SVerticalBox::Slot()
+		.AutoHeight()
+		[
+			SAssignNew(MediaObjectInfoBox, SHorizontalBox)
+		]
+	];
+}
+
+void SMediaObjectInfoPanel::SetMediaObject(UObject* InMediaObject)
+{
+	MediaObject = InMediaObject;
+	ClearInfo();
+	FillInfo();
+
+	if (!MediaObject.IsValid())
+	{
+		SetVisibility(EVisibility::Collapsed);
+	}
+	else
+	{
+		SetVisibility(EVisibility::Visible);
+	}
+}
+
+UMediaSource* SMediaObjectInfoPanel::ObjectAsMediaSource() const
+{
+	return Cast<UMediaSource>(MediaObject);
+}
+
+UMediaOutput* SMediaObjectInfoPanel::ObjectAsMediaOutput() const
+{
+	return Cast<UMediaOutput>(MediaObject);
+}
+
+void SMediaObjectInfoPanel::ClearInfo()
+{
+	BasicInfo.ClearEntries();
+	TypeInfo.ClearEntries();
+	CaptureInfo.ClearEntries();
+	
+	MediaObjectInfoBox->ClearChildren();
+}
+
+void SMediaObjectInfoPanel::FillInfo()
+{
+	if (!MediaObject.IsValid() || !MediaProfile.IsValid())
+	{
+		return;
+	}
+
+	UObject* ActualMediaObject = MediaObject.Get();
+	if (UProxyMediaSource* ProxySource = Cast<UProxyMediaSource>(ActualMediaObject))
+	{
+		ActualMediaObject = ProxySource->GetLeafMediaSource();
+	}
+	else if (UProxyMediaOutput* ProxyOutput = Cast<UProxyMediaOutput>(ActualMediaObject))
+	{
+		ActualMediaObject = ProxyOutput->GetLeafMediaOutput();
+	}
+		
+	FText Label;
+	if (UMediaSource* MediaSource = ObjectAsMediaSource())
+	{
+		int32 Index = MediaProfile->FindMediaSourceIndex(MediaSource);
+		if (MediaSource->IsA<UDummyMediaSource>())
+		{
+			Index = Cast<UDummyMediaSource>(MediaSource)->MediaProfileIndex;
+		}
+		
+		Label = FText::FromString(MediaProfile->GetLabelForMediaSource(Index));
+	}
+	else if (UMediaOutput* MediaOutput = ObjectAsMediaOutput())
+	{
+		int32 Index = MediaProfile->FindMediaOutputIndex(MediaOutput);
+		if (MediaOutput->IsA<UDummyMediaOutput>())
+		{
+			Index = Cast<UDummyMediaOutput>(MediaOutput)->MediaProfileIndex;
+		}
+		
+		Label = FText::FromString(MediaProfile->GetLabelForMediaOutput(Index));
+	}
+		
+	BasicInfo.AddEntry(LOCTEXT("MediaObjectNameLabel", "Label:"), Label);
+
+	const bool bIsProxy = MediaObject->IsA<UProxyMediaSource>() || MediaObject->IsA<UProxyMediaOutput>();
+	BasicInfo.AddEntry(LOCTEXT("MediaObjectIsProxyLabel", "Is Proxy:"), bIsProxy ? LOCTEXT("Yes", "Yes") : LOCTEXT("No", "No"));
+
+	if (ActualMediaObject)
+	{
+		TypeInfo.AddEntry(ActualMediaObject->GetClass()->GetDisplayNameText(), FSlateIconFinder::FindIconForClass(ActualMediaObject->GetClass()).GetIcon());
+	}
+	else
+	{
+		const FText NoObjectSet = ObjectAsMediaSource() ? LOCTEXT("NoMediaSourceSet", "No media source set") : LOCTEXT("NoMediaOutputSet", "No media output set");
+		TypeInfo.AddEntry(NoObjectSet);
+	}
+		
+	TArray<TTuple<FText, FText, FText>> InfoElements;
+
+	if (UMediaSource* MediaSource = Cast<UMediaSource>(ActualMediaObject))
+	{
+		MediaSource->GetDetailsPanelInfoElements(InfoElements);
+	}
+	else if (UMediaOutput* MediaOutput = Cast<UMediaOutput>(ActualMediaObject))
+	{
+		MediaOutput->GetDetailsPanelInfoElements(InfoElements);
+	}
+		
+	TMap<FString, FTwoColumnInfo> InfoWidgets;
+	for (TTuple<FText, FText, FText>& InfoElement : InfoElements)
+	{
+		const FText InfoGroup = InfoElement.Get<0>();
+		const FText InfoLabel = InfoElement.Get<1>();
+		const FText InfoValue = InfoElement.Get<2>();
+			
+		if (!InfoWidgets.Contains(InfoGroup.ToString()))
+		{
+			InfoWidgets.Add(InfoGroup.ToString(), FTwoColumnInfo());
+			MediaObjectInfoBox->AddSlot()
+			[
+				CreateInfoGroupWidget(InfoGroup, InfoWidgets[InfoGroup.ToString()])
+			];
+		}
+
+		InfoWidgets[InfoGroup.ToString()].AddEntry(FText::Format(LOCTEXT("MediaObjectInfoElementLabel", "{0}:"), InfoLabel), InfoValue);
+	}
+	
+	if (Cast<UMediaOutput>(ActualMediaObject))
+	{
+		MediaObjectInfoBox->AddSlot()
+		[
+			CreateInfoGroupWidget(LOCTEXT("CaptureSettingsGroupLabel", "Capture Settings"), CaptureInfo)
+		];
+		
+		CaptureInfo.AddEntry(LOCTEXT("CaptureMethodLabel", "Capture Method:"), TAttribute<FText>::CreateSP(this, &SMediaObjectInfoPanel::GetCaptureMethodText));
+		CaptureInfo.AddEntry(TAttribute<FText>::CreateSP(this, &SMediaObjectInfoPanel::GetCaptureObjectLabelText), TAttribute<FText>::CreateSP(this, &SMediaObjectInfoPanel::GetCaptureObjectValueText));
+		CaptureInfo.AddEntry(LOCTEXT("CaptureStatsLabel", "Status:"),
+			TAttribute<FText>::CreateSP(this, &SMediaObjectInfoPanel::GetCaptureStatusText),
+			TAttribute<FSlateColor>::CreateSP(this, &SMediaObjectInfoPanel::GetCaptureStatusColor));
+		CaptureInfo.AddEntry(LOCTEXT("CaptureColorConversionLabel", "Color Conversion:"), TAttribute<FText>::CreateSP(this, &SMediaObjectInfoPanel::GetCaptureColorConversionText));
+	}
+}
+
+TSharedRef<SWidget> SMediaObjectInfoPanel::CreateInfoGroupWidget(const TAttribute<FText>& InHeader, FTwoColumnInfo& InInfo)
+{
+	return SNew(SVerticalBox)
+
+		+SVerticalBox::Slot()
+		.AutoHeight()
+		[
+			SNew(STextBlock)
+			.Text(InHeader)
+			.Font(IDetailLayoutBuilder::GetDetailFontBold())
+		]
+
+		+SVerticalBox::Slot()
+		.AutoHeight()
+		[
+			InInfo.CreateWidget()
+		];			
+}
+
+FText SMediaObjectInfoPanel::GetCaptureMethodText() const
+{
+	UMediaOutput* MediaOutput = ObjectAsMediaOutput();
+	UMediaFrameworkWorldSettingsAssetUserData* CaptureSettings = FMediaProfileEditor::FindOrAddMediaFrameworkCaptureSettings();
+	if (!MediaOutput || !CaptureSettings)
+	{
+		return FText::GetEmpty();
+	}
+
+	bool bHasViewportCaptures = CaptureSettings->ViewportCaptures.ContainsByPredicate([MediaOutput](const FMediaFrameworkCaptureCameraViewportCameraOutputInfo& Info)
+	{
+		return Info.MediaOutput == MediaOutput;
+	});
+
+	bool bHasRenderTargetCaptures = CaptureSettings->RenderTargetCaptures.ContainsByPredicate([MediaOutput](const FMediaFrameworkCaptureRenderTargetCameraOutputInfo& Info)
+	{
+		return Info.MediaOutput == MediaOutput;
+	});
+
+	if (bHasViewportCaptures && bHasRenderTargetCaptures)
+	{
+		return LOCTEXT("MultipleCapturesValue", "Multiple Captures");
+	}
+
+	if (bHasViewportCaptures)
+	{
+		return LOCTEXT("ViewportCaptureValue", "Viewport");
+	}
+
+	if (bHasRenderTargetCaptures)
+	{
+		return LOCTEXT("RenderTargetCaptureValue", "Render Target");
+	}
+
+	return LOCTEXT("NoCapturesValues", "No Captures Configured");
+}
+
+FText SMediaObjectInfoPanel::GetCaptureObjectLabelText() const
+{
+	UMediaOutput* MediaOutput = ObjectAsMediaOutput();
+	UMediaFrameworkWorldSettingsAssetUserData* CaptureSettings = FMediaProfileEditor::FindOrAddMediaFrameworkCaptureSettings();
+	if (!MediaOutput || !CaptureSettings)
+	{
+		return FText::GetEmpty();
+	}
+
+	bool bHasViewportCaptures = CaptureSettings->ViewportCaptures.ContainsByPredicate([MediaOutput](const FMediaFrameworkCaptureCameraViewportCameraOutputInfo& Info)
+	{
+		return Info.MediaOutput == MediaOutput;
+	});
+
+	bool bHasRenderTargetCaptures = CaptureSettings->RenderTargetCaptures.ContainsByPredicate([MediaOutput](const FMediaFrameworkCaptureRenderTargetCameraOutputInfo& Info)
+	{
+		return Info.MediaOutput == MediaOutput;
+	});
+
+	if (bHasViewportCaptures && bHasRenderTargetCaptures)
+	{
+		return LOCTEXT("MultipleCaptureObjectsLabel", "Multiple Objects:");
+	}
+
+	if (bHasViewportCaptures)
+	{
+		return LOCTEXT("CameraObjectLabel", "Cameras:");
+	}
+
+	if (bHasRenderTargetCaptures)
+	{
+		return LOCTEXT("RenderTargetObjectLabel", "Render Target:");
+	}
+
+	return LOCTEXT("NoObjectsLabel", "No Objects:");
+}
+
+FText SMediaObjectInfoPanel::GetCaptureObjectValueText() const
+{
+	UMediaOutput* MediaOutput = ObjectAsMediaOutput();
+	UMediaFrameworkWorldSettingsAssetUserData* CaptureSettings = FMediaProfileEditor::FindOrAddMediaFrameworkCaptureSettings();
+	if (!MediaOutput || !CaptureSettings)
+	{
+		return FText::GetEmpty();
+	}
+
+	bool bHasViewportCaptures = false;
+	FText CamerasText;
+	for (const FMediaFrameworkCaptureCameraViewportCameraOutputInfo& Info : CaptureSettings->ViewportCaptures)
+	{
+		if (Info.MediaOutput != MediaOutput)
+		{
+			continue;
+		}
+
+		bHasViewportCaptures = true;
+		
+		for (const TLazyObjectPtr<AActor>& Camera : Info.LockedActors)
+		{
+			AActor* Actor = Camera.Get();
+			if (Actor)
+			{
+				if (CamerasText.IsEmpty())
+				{
+					CamerasText = FText::FromString(Actor->GetActorNameOrLabel());
+				}
+				else
+				{
+					CamerasText = FText::Format(LOCTEXT("AppendCameraNameFormat", "{0}, {1}"), CamerasText, FText::FromString(Actor->GetActorNameOrLabel()));
+				}
+			}
+		}
+	}
+
+	bool bHasRenderTargetCaptures = false;
+	FText RenderTargetsText;
+	for (const FMediaFrameworkCaptureRenderTargetCameraOutputInfo& Info : CaptureSettings->RenderTargetCaptures)
+	{
+		if (Info.MediaOutput != MediaOutput)
+		{
+			continue;
+		}
+
+		bHasRenderTargetCaptures = true;
+		
+		
+		FText RenderTargetText = Info.RenderTarget ? FText::FromString(Info.RenderTarget->GetName()) : LOCTEXT("NoRenderTarget", "No Render Target");
+		
+		if (RenderTargetsText.IsEmpty())
+		{
+			RenderTargetsText = RenderTargetText;
+		}
+		else
+		{
+			RenderTargetsText = FText::Format(LOCTEXT("AppendRenderTargetNameFormat", "{0}, {1}"), RenderTargetsText, RenderTargetText);
+		}
+	}
+
+	if (bHasViewportCaptures && bHasRenderTargetCaptures)
+	{
+		return FText::GetEmpty();
+	}
+
+	if (bHasViewportCaptures)
+	{
+		return CamerasText;
+	}
+
+	if (bHasRenderTargetCaptures)
+	{
+		return RenderTargetsText;
+	}
+
+	return FText::GetEmpty();
+}
+
+FText SMediaObjectInfoPanel::GetCaptureStatusText() const
+{
+	UMediaOutput* MediaOutput = ObjectAsMediaOutput();
+	if (!MediaOutput || !MediaProfileEditor.IsValid())
+	{
+		return LOCTEXT("UnknownCaptureStatusText", "Unknown");
+	}
+
+	return MediaProfileEditor.Pin()->IsMediaOutputCapturing(MediaOutput) ? LOCTEXT("CapturingStatusText", "Capturing") : LOCTEXT("NotCapturingStatusText", "Not Capturing");
+}
+
+FSlateColor SMediaObjectInfoPanel::GetCaptureStatusColor() const
+{
+	UMediaOutput* MediaOutput = ObjectAsMediaOutput();
+	if (!MediaOutput || !MediaProfileEditor.IsValid())
+	{
+		return FSlateColor::UseForeground();
+	}
+
+	return MediaProfileEditor.Pin()->IsMediaOutputCapturing(MediaOutput) ? FSlateColor(FLinearColor(0.0, 1.0, 0.0)) : FSlateColor::UseForeground();
+}
+
+FText SMediaObjectInfoPanel::GetCaptureColorConversionText() const
+{
+	UMediaOutput* MediaOutput = ObjectAsMediaOutput();
+	UMediaFrameworkWorldSettingsAssetUserData* CaptureSettings = FMediaProfileEditor::FindOrAddMediaFrameworkCaptureSettings();
+	if (!MediaOutput || !CaptureSettings)
+	{
+		return LOCTEXT("NoneText", "None");
+	}
+
+	FText ColorConversionText;
+	for (const FMediaFrameworkCaptureCameraViewportCameraOutputInfo& Info : CaptureSettings->ViewportCaptures)
+	{
+		if (Info.MediaOutput != MediaOutput)
+		{
+			continue;
+		}
+
+		if (ColorConversionText.IsEmpty())
+		{
+			ColorConversionText = FText::FromString(Info.CaptureOptions.GetColorConversionSettingsString());
+		}
+		else
+		{
+			ColorConversionText = FText::Format(LOCTEXT("AppendColorConversionFormat", "{0}, {1}"), ColorConversionText, FText::FromString(Info.CaptureOptions.GetColorConversionSettingsString()));
+		}
+	}
+	
+	for (const FMediaFrameworkCaptureRenderTargetCameraOutputInfo& Info : CaptureSettings->RenderTargetCaptures)
+	{
+		if (Info.MediaOutput != MediaOutput)
+		{
+			continue;
+		}
+
+		if (ColorConversionText.IsEmpty())
+		{
+			ColorConversionText = FText::FromString(Info.CaptureOptions.GetColorConversionSettingsString());
+		}
+		else
+		{
+			ColorConversionText = FText::Format(LOCTEXT("AppendColorConversionFormat", "{0}, {1}"), ColorConversionText, FText::FromString(Info.CaptureOptions.GetColorConversionSettingsString()));
+		}
+	}
+
+	if (ColorConversionText.IsEmpty())
+	{
+		return LOCTEXT("NoneText", "None");
+	}
+
+	return ColorConversionText;
+}
+
+#undef LOCTEXT_NAMESPACE
diff --git a/Engine/Plugins/Media/MediaFrameworkUtilities/Source/MediaFrameworkUtilitiesEditor/Private/AssetEditor/SMediaProfileDetailsInfoPanel.h b/Engine/Plugins/Media/MediaFrameworkUtilities/Source/MediaFrameworkUtilitiesEditor/Private/AssetEditor/SMediaProfileDetailsInfoPanel.h
new file mode 100644
--- /dev/null
+++ b/Engine/Plugins/Media/MediaFrameworkUtilities/Source/MediaFrameworkUtilitiesEditor/Private/AssetEditor/SMediaProfileDetailsInfoPanel.h
@@ -0,0 +1,84 @@
+﻿// Copyright Epic Games, Inc. All Rights Reserved.
+
+#pragma once
+
+#include "CoreMinimal.h"
+#include "Widgets/SCompoundWidget.h"
+
+class FMediaProfileEditor;
+class SVerticalBox;
+class SHorizontalBox;
+class UMediaOutput;
+class UMediaSource;
+class UMediaProfile;
+
+/**
+ * Info panel that is displayed above the details panel, and displays a summary of a media item's configuration
+ */
+class SMediaObjectInfoPanel : public SCompoundWidget
+{
+private:
+	/** Helper struct to maintain two columns of widgets, one for labels and one for value content */
+	struct FTwoColumnInfo
+	{
+	public:
+		TSharedRef<SWidget> CreateWidget();
+
+		void AddEntry(TAttribute<FText> Label, TAttribute<FText> Value, TAttribute<FSlateColor> ValueColor = TAttribute<FSlateColor>());
+		void AddEntry(TAttribute<FText> Value, const FSlateBrush* Icon);
+		void AddEntry(TAttribute<FText> Value);
+
+		void ClearEntries();
+
+	private:
+		TSharedPtr<SVerticalBox> LabelColumn;
+		TSharedPtr<SVerticalBox> ValueColumn;
+	};
+	
+public:
+	SLATE_BEGIN_ARGS(SMediaObjectInfoPanel) {}
+	SLATE_END_ARGS()
+	
+	void Construct(const FArguments& InArgs, TWeakPtr<FMediaProfileEditor>& InMediaProfileEditor, UMediaProfile* InMediaProfile);
+
+	/** Sets the media source or output whose information is being displayed in the panel */
+	void SetMediaObject(UObject* InMediaObject);
+
+private:
+	UMediaSource* ObjectAsMediaSource() const;
+	UMediaOutput* ObjectAsMediaOutput() const;
+	
+	void ClearInfo();
+	void FillInfo();
+
+	TSharedRef<SWidget> CreateInfoGroupWidget(const TAttribute<FText>& InHeader, FTwoColumnInfo& InInfo);
+
+	FText GetCaptureMethodText() const;
+	FText GetCaptureObjectLabelText() const;
+	FText GetCaptureObjectValueText() const;
+	FText GetCaptureStatusText() const;
+	FSlateColor GetCaptureStatusColor() const;
+	FText GetCaptureColorConversionText() const;
+	
+private:
+	/** Media profile editor that owns this widget */
+	TWeakPtr<FMediaProfileEditor> MediaProfileEditor;
+
+	/** Media profile that is being edited */
+	TWeakObjectPtr<UMediaProfile> MediaProfile;
+
+	/** Media source or output whose info is being displayed */
+	TWeakObjectPtr<UObject> MediaObject;
+
+	/** Container for any extra columns of info */
+	TSharedPtr<SHorizontalBox> MediaObjectInfoBox;
+
+	/** Basic info column, which displays media object label and proxy status */
+	FTwoColumnInfo BasicInfo;
+
+	/** Column that displays details about the type of the media object */
+	FTwoColumnInfo TypeInfo;
+
+	/** Column that displays details about the media capture settings if the object is a media output */
+	FTwoColumnInfo CaptureInfo;
+};
diff --git a/Engine/Plugins/Media/MediaFrameworkUtilities/Source/MediaFrameworkUtilitiesEditor/Private/AssetEditor/SMediaProfileDetailsPanel.cpp b/Engine/Plugins/Media/MediaFrameworkUtilities/Source/MediaFrameworkUtilitiesEditor/Private/AssetEditor/SMediaProfileDetailsPanel.cpp
new file mode 100644
--- /dev/null
+++ b/Engine/Plugins/Media/MediaFrameworkUtilities/Source/MediaFrameworkUtilitiesEditor/Private/AssetEditor/SMediaProfileDetailsPanel.cpp
@@ -0,0 +1,175 @@
+﻿// Copyright Epic Games, Inc. All Rights Reserved.
+
+#include "SMediaProfileDetailsPanel.h"
+
+#include "ClassViewerFilter.h"
+#include "ClassViewerModule.h"
+#include "DetailCategoryBuilder.h"
+#include "DetailLayoutBuilder.h"
+#include "DummyMediaObject.h"
+#include "IDetailsView.h"
+#include "MediaSource.h"
+#include "MediaOutput.h"
+#include "MediaProfileMediaItemDetailCustomization.h"
+#include "SMediaProfileDetailsInfoPanel.h"
+#include "CaptureTab/SMediaFrameworkCapture.h"
+#include "Modules/ModuleManager.h"
+#include "Profile/MediaProfile.h"
+#include "Widgets/Input/SComboButton.h"
+#include "Widgets/Input/SEditableTextBox.h"
+#include "Widgets/Layout/SSplitter.h"
+
+#define LOCTEXT_NAMESPACE "MediaProfileDetailsPanel"
+
+void SMediaProfileDetailsPanel::Construct(const FArguments& InArgs, const TSharedPtr<FMediaProfileEditor> InOwningEditor, UMediaProfile* InMediaProfile)
+{
+	OwningEditor = InOwningEditor;
+	MediaProfile = InMediaProfile;
+
+	FPropertyEditorModule& PropertyModule = FModuleManager::LoadModuleChecked<FPropertyEditorModule>("PropertyEditor");
+	FDetailsViewArgs DetailsArgs;
+	DetailsArgs.bLockable = true;
+	DetailsArgs.NameAreaSettings = FDetailsViewArgs::ObjectsUseNameArea | FDetailsViewArgs::ComponentsAndActorsUseNameArea;
+	DetailsArgs.bCustomNameAreaLocation = true;
+	DetailsArgs.bCustomFilterAreaLocation = true;
+	DetailsArgs.DefaultsOnlyVisibility = EEditDefaultsOnlyNodeVisibility::Hide;
+	DetailsArgs.bShowSectionSelector = true;
+	
+	DetailsView = PropertyModule.CreateDetailView(DetailsArgs);
+	DetailsView->RegisterInstancedCustomPropertyLayout(UMediaSource::StaticClass(), FOnGetDetailCustomizationInstance::CreateLambda([this]
+	{
+		return MakeShared<FMediaProfileMediaSourceDetailCustomization>(OwningEditor, MediaProfile.Get(), FSimpleDelegate::CreateSP(this, &SMediaProfileDetailsPanel::RefreshDetailsView));
+	}));
+	DetailsView->RegisterInstancedCustomPropertyLayout(UMediaOutput::StaticClass(), FOnGetDetailCustomizationInstance::CreateLambda([this]
+	{
+		return MakeShared<FMediaProfileMediaOutputDetailCustomization>(OwningEditor, MediaProfile.Get(), FSimpleDelegate::CreateSP(this, &SMediaProfileDetailsPanel::RefreshDetailsView));
+	}));
+	
+	ChildSlot
+	[
+		SNew(SVerticalBox)
+		+SVerticalBox::Slot()
+		.Padding(10.f, 4.f, 0.f, 0.f)
+		.AutoHeight()
+		[
+			DetailsView->GetNameAreaWidget().ToSharedRef()
+		]
+
+		+SVerticalBox::Slot()
+		.Padding(10.0f, 16.0f)
+		.AutoHeight()
+		[
+			SAssignNew(InfoPanel, SMediaObjectInfoPanel, OwningEditor, MediaProfile.Get())
+		]
+		
+		+SVerticalBox::Slot()
+		[
+			SNew(SSplitter)
+			.MinimumSlotHeight(80.0f)
+			.Orientation(Orient_Vertical)
+			.Style(FAppStyle::Get(), "SplitterDark")
+			.PhysicalSplitterHandleSize(2.0f)
+			+SSplitter::Slot()
+			[
+				SNew( SVerticalBox )
+				+SVerticalBox::Slot()
+				.AutoHeight()
+				[
+					DetailsView->GetFilterAreaWidget().ToSharedRef()
+				]
+				+SVerticalBox::Slot()
+				[
+					DetailsView.ToSharedRef()
+				]
+			]
+		]
+	];
+}
+
+void SMediaProfileDetailsPanel::SetSelectedMediaItems(const TArray<int32>& InSelectedMediaSourceIndices, const TArray<int32>& InSelectedMediaOutputIndices)
+{
+	SelectedMediaSources.Empty(InSelectedMediaSourceIndices.Num());
+	SelectedMediaSources.Append(InSelectedMediaSourceIndices);
+	
+	SelectedMediaOutputs.Empty(InSelectedMediaOutputIndices.Num());
+	SelectedMediaOutputs.Append(InSelectedMediaOutputIndices);
+	
+	RefreshDetailsView();
+}
+
+void SMediaProfileDetailsPanel::RefreshDetailsView()
+{
+	DummyMediaObjects.Empty();
+	
+	if (!MediaProfile.IsValid())
+	{
+		DetailsView->SetObject(nullptr);
+		return;
+	}
+	
+	// Can't display details from sources and outputs at the same time, so clear and return
+	if (SelectedMediaSources.Num() > 0 && SelectedMediaOutputs.Num() > 0)
+	{
+		DetailsView->SetObject(nullptr);
+		InfoPanel->SetMediaObject(nullptr);
+		return;
+	}
+
+	TArray<UObject*> Objects;
+	if (SelectedMediaSources.Num() > 0)
+	{
+		Objects.Reserve(SelectedMediaSources.Num());
+		for (int Index = 0; Index < SelectedMediaSources.Num(); ++Index)
+		{
+			UMediaSource* MediaSource = MediaProfile->GetMediaSource(SelectedMediaSources[Index]);
+			if (!MediaSource)
+			{
+				UDummyMediaSource* DummyMediaSource = NewObject<UDummyMediaSource>();
+				DummyMediaObjects.Add(TStrongObjectPtr<UObject>(DummyMediaSource));
+				DummyMediaSource->MediaProfileIndex = SelectedMediaSources[Index];
+
+				MediaSource = DummyMediaSource;
+			}
+			
+			Objects.Add(MediaSource);
+		}
+	}
+	else if (SelectedMediaOutputs.Num() > 0)
+	{
+		Objects.Reserve(SelectedMediaOutputs.Num());
+		for (int Index = 0; Index < SelectedMediaOutputs.Num(); ++Index)
+		{
+			UMediaOutput* MediaOutput = MediaProfile->GetMediaOutput(SelectedMediaOutputs[Index]);
+			if (!MediaOutput)
+			{
+				UDummyMediaOutput* DummyMediaOutput = NewObject<UDummyMediaOutput>();
+				DummyMediaObjects.Add(TStrongObjectPtr<UObject>(DummyMediaOutput));
+				DummyMediaOutput->MediaProfileIndex = SelectedMediaOutputs[Index];
+
+				MediaOutput = DummyMediaOutput;
+			}
+			
+			Objects.Add(MediaOutput);
+		}
+	}
+
+	if (Objects.Num() > 0)
+	{
+		DetailsView->SetObjects(Objects);
+	}
+	else
+	{
+		DetailsView->SetObject(nullptr);
+	}
+	
+	if (Objects.Num() == 1 && !(Objects[0]->IsA<UDummyMediaSource>() || Objects[0]->IsA<UDummyMediaOutput>()))
+	{
+		InfoPanel->SetMediaObject(Objects[0]);
+	}
+	else
+	{
+		InfoPanel->SetMediaObject(nullptr);
+	}
+}
+
+#undef LOCTEXT_NAMESPACE
\ No newline at end of file
diff --git a/Engine/Plugins/Media/MediaFrameworkUtilities/Source/MediaFrameworkUtilitiesEditor/Private/AssetEditor/SMediaProfileDetailsPanel.h b/Engine/Plugins/Media/MediaFrameworkUtilities/Source/MediaFrameworkUtilitiesEditor/Private/AssetEditor/SMediaProfileDetailsPanel.h
new file mode 100644
--- /dev/null
+++ b/Engine/Plugins/Media/MediaFrameworkUtilities/Source/MediaFrameworkUtilitiesEditor/Private/AssetEditor/SMediaProfileDetailsPanel.h
@@ -0,0 +1,59 @@
+﻿// Copyright Epic Games, Inc. All Rights Reserved.
+
+#pragma once
+
+#include "CoreMinimal.h"
+#include "MediaProfileEditor.h"
+#include "MediaSource.h"
+#include "Widgets/DeclarativeSyntaxSupport.h"
+#include "Widgets/SCompoundWidget.h"
+
+class SMediaObjectInfoPanel;
+class IDetailsView;
+class UMediaProfile;
+
+/** Details panel for the media profile's selected media sources and outputs */
+class SMediaProfileDetailsPanel : public SCompoundWidget
+{
+public:
+	DECLARE_DELEGATE(FOnDetailsChanged);
+	
+public:
+	SLATE_BEGIN_ARGS(SMediaProfileDetailsPanel) {}
+		SLATE_EVENT(FOnDetailsChanged, OnDetailsChanged)
+	SLATE_END_ARGS()
+	
+	void Construct(const FArguments& InArgs, const TSharedPtr<FMediaProfileEditor> InOwningEditor, UMediaProfile* InMediaProfile);
+
+	/**
+	 * Sets the media items in the media profile to display in the details panel
+	 * @param InSelectedMediaSourceIndices A list of indices of the media sources in the media profile to display
+	 * @param InSelectedMediaOutputIndices A list of indices of the media outputs in the media profile to display
+	 */
+	void SetSelectedMediaItems(const TArray<int32>& InSelectedMediaSourceIndices, const TArray<int32>& InSelectedMediaOutputIndices);
+
+private:
+	void RefreshDetailsView();
+	
+private:
+	/** The editor that owns this details panel */
+	TWeakPtr<FMediaProfileEditor> OwningEditor;
+	
+	/** The media profile being edited by this editor */
+	TWeakObjectPtr<UMediaProfile> MediaProfile;
+	
+	/** Details view to display the media source or output properties */
+	TSharedPtr<IDetailsView> DetailsView;
+
+	/** Panel to display info on the selected media objects */
+	TSharedPtr<SMediaObjectInfoPanel> InfoPanel;
+
+	/** Strong pointers for any dummy objects needed to represent null media entries in the media profile */
+	TArray<TStrongObjectPtr<UObject>> DummyMediaObjects;
+
+	/** Indices of the selected media sources within the media profile to display in the details panel */
+	TArray<int32> SelectedMediaSources;
+
+	/** Indices of the selected media outputs within the media profile to display in the details panel */
+	TArray<int32> SelectedMediaOutputs;
+};
diff --git a/Engine/Plugins/Media/MediaFrameworkUtilities/Source/MediaFrameworkUtilitiesEditor/Private/AssetEditor/SMediaProfileSourcesTreeView.cpp b/Engine/Plugins/Media/MediaFrameworkUtilities/Source/MediaFrameworkUtilitiesEditor/Private/AssetEditor/SMediaProfileSourcesTreeView.cpp
--- a/Engine/Plugins/Media/MediaFrameworkUtilities/Source/MediaFrameworkUtilitiesEditor/Private/AssetEditor/SMediaProfileSourcesTreeView.cpp
+++ b/Engine/Plugins/Media/MediaFrameworkUtilities/Source/MediaFrameworkUtilitiesEditor/Private/AssetEditor/SMediaProfileSourcesTreeView.cpp
@@ -10,6 +10,7 @@
 #include "SPositiveActionButton.h"
 #include "SSettingsEditorCheckoutNotice.h"
 #include "UnrealExporter.h"
+#include "CaptureTab/SMediaFrameworkCapture.h"
 #include "DragAndDrop/DecoratedDragDropOp.h"
 #include "Exporters/Exporter.h"
 #include "Filters/GenericFilter.h"
@@ -261,7 +262,15 @@ private:
 
 		if (OnMoveMediaItem.IsBound())
 		{
-			const int32 DestIndex = InItemDropZone == EItemDropZone::AboveItem ? InTreeItem->Index: InTreeItem->Index + 1;
+			int32 DestIndex = InItemDropZone == EItemDropZone::AboveItem ? InTreeItem->Index: InTreeItem->Index + 1;
+
+			// If the destination is further down than the item being moved, we must subtract one from the destination index
+			// to account for the removal of the original item
+			if (InTreeItem->Index > DragDropOp->DraggedItem.Pin()->Index)
+			{
+				--DestIndex;
+			}
+			
 			OnMoveMediaItem.Execute(DragDropOp->DraggedItem.Pin(), DestIndex);
 			return FReply::Handled();
 		}
@@ -618,6 +627,7 @@ private:
 void SMediaProfileSourcesTreeView::Construct(const FArguments& InArgs, UMediaProfile* InMediaProfile)
 {
 	MediaProfile = InMediaProfile;
+	OnSelectedMediaItemsChanged = InArgs._OnSelectedMediaItemsChanged;
 
 	CommandList = MakeShared<FUICommandList>();
 	BindCommands();
@@ -666,6 +676,7 @@ void SMediaProfileSourcesTreeView::Construct(const FArguments& InArgs, UMediaPro
 			.OnGetChildren(this, &SMediaProfileSourcesTreeView::GetTreeItemChildren)
 			.OnGenerateRow(this, &SMediaProfileSourcesTreeView::GenerateTreeItemRow)
 			.OnContextMenuOpening(this, &SMediaProfileSourcesTreeView::CreateContextMenu)
+			.OnSelectionChanged(this, &SMediaProfileSourcesTreeView::OnMediaTreeSelectionChanged)
 		]
 	];
 
@@ -1010,12 +1021,19 @@ void SMediaProfileSourcesTreeView::AddNewMediaOutput()
 
 void SMediaProfileSourcesTreeView::FillMediaTree()
 {
+	// In case a media object changed type, store the current selection through their indexes to re-select them correctly after
+	// the tree has been recreated
+	TArray<int32> SelectedMediaSources;
+	TArray<int32> SelectedMediaOutputs;
+	GetSelectedMediaIndexes(SelectedMediaSources, SelectedMediaOutputs, true);
+	
 	MediaTreeItems.Empty();
 
 	if (MediaProfile.IsValid())
 	{
 		FMediaTreeItemPtr MediaSourcesTreeItem = MakeShared<FMediaTreeItem>();
 		MediaSourcesTreeItem->Label = LOCTEXT("MediaSourcesTreeItemLabel", "Media Sources");
+		MediaSourcesTreeItem->BackingClass = UMediaSource::StaticClass();
 		
 		for (int32 Index = 0; Index < MediaProfile->NumMediaSources(); ++Index)
 		{
@@ -1035,7 +1053,7 @@ void SMediaProfileSourcesTreeView::FillMediaTree()
 
 			if (MediaSource)
 			{
-				MediaSourceTreeItem->Label = FText::FromString(MediaSource->GetName());
+				MediaSourceTreeItem->Label = FText::FromString(MediaProfile->GetLabelForMediaSource(Index));
 				MediaSourceTreeItem->Type = MediaSource->GetClass()->GetDisplayNameText();
 				MediaSourceTreeItem->Icon = FSlateIconFinder::FindIconForClass(MediaSource->GetClass()).GetIcon();
 				MediaSourceTreeItem->Configuration = FText::FromString(MediaSource->GetDescriptionString());
@@ -1055,6 +1073,7 @@ void SMediaProfileSourcesTreeView::FillMediaTree()
 		
 		FMediaTreeItemPtr MediaOutputsTreeItem = MakeShared<FMediaTreeItem>();
 		MediaOutputsTreeItem->Label = LOCTEXT("MediaOutputsTreeItemLabel", "Media Outputs");
+		MediaOutputsTreeItem->BackingClass = UMediaOutput::StaticClass();
 		
 		for (int32 Index = 0; Index < MediaProfile->NumMediaOutputs(); ++Index)
 		{
@@ -1074,7 +1093,7 @@ void SMediaProfileSourcesTreeView::FillMediaTree()
 
 			if (MediaOutput)
 			{
-				MediaOutputTreeItem->Label = FText::FromString(MediaOutput->GetName());
+				MediaOutputTreeItem->Label = FText::FromString(MediaProfile->GetLabelForMediaOutput(Index));
 				MediaOutputTreeItem->Type = MediaOutput->GetClass()->GetDisplayNameText();
 				MediaOutputTreeItem->Icon = FSlateIconFinder::FindIconForClass(MediaOutput->GetClass()).GetIcon();
 				MediaOutputTreeItem->Configuration = FText::FromString(MediaOutput->GetDescriptionString());
@@ -1101,6 +1120,65 @@ void SMediaProfileSourcesTreeView::FillMediaTree()
 		{
 			MediaTreeView->SetItemExpansion(TreeItem, true);
 		}
+
+		if (MediaTreeItems.Num() > 0)
+		{
+			for (int32 SourceIndex : SelectedMediaSources)
+			{
+				if (SourceIndex == INDEX_NONE)
+				{
+					MediaTreeView->SetItemSelection(MediaTreeItems[0], true);
+				}
+				else if (MediaTreeItems[0]->Children.IsValidIndex(SourceIndex))
+				{
+					MediaTreeView->SetItemSelection(MediaTreeItems[0]->Children[SourceIndex], true);
+				}
+			}
+
+			for (int32 OutputIndex : SelectedMediaOutputs)
+			{
+				if (OutputIndex == INDEX_NONE)
+				{
+					MediaTreeView->SetItemSelection(MediaTreeItems[1], true);
+				}
+				else if (MediaTreeItems[1]->Children.IsValidIndex(OutputIndex))
+				{
+					MediaTreeView->SetItemSelection(MediaTreeItems[1]->Children[OutputIndex], true);
+				}
+			}
+		}
+	}
+}
+
+void SMediaProfileSourcesTreeView::GetSelectedMediaIndexes(TArray<int32>& OutSelectedMediaSources, TArray<int32>& OutSelectedMediaOutputs, bool bIncludeGroupIndex)
+{
+	if (!MediaTreeView.IsValid())
+	{
+		return;
+	}
+
+	TArray<FMediaTreeItemPtr> SelectedItems = MediaTreeView->GetSelectedItems();
+
+	for (const FMediaTreeItemPtr& Item : SelectedItems)
+	{
+		if (!Item.IsValid())
+		{
+			continue;
+		}
+
+		if (!Item->IsLeaf() && !bIncludeGroupIndex)
+		{
+			continue;
+		}
+
+		if (Item->BackingClass->IsChildOf<UMediaSource>())
+		{
+			OutSelectedMediaSources.Add(Item->Index);
+		}
+		else if (Item->BackingClass->IsChildOf<UMediaOutput>())
+		{
+			OutSelectedMediaOutputs.Add(Item->Index);
+		}
 	}
 }
 
@@ -1146,11 +1224,27 @@ TSharedPtr<SWidget> SMediaProfileSourcesTreeView::CreateContextMenu()
 
 void SMediaProfileSourcesTreeView::OnObjectPropertyChanged(UObject* InObject, FPropertyChangedEvent& InPropertyChangedEvent)
 {
-	if (!InObject || InObject != MediaProfile)
+	if (!MediaProfile.IsValid())
 	{
 		return;
 	}
 	
+	if (!InObject)
+	{
+		return;
+	}
+
+	if (MediaProfile != InObject)
+	{
+		// If the object isn't the media profile itself, check to see if it is one of the profile's sources or outputs,
+		// as we want to update the tree view if the sources or outputs change as well
+		if (MediaProfile->FindMediaSourceIndex(Cast<UMediaSource>(InObject)) == INDEX_NONE &&
+			MediaProfile->FindMediaOutputIndex(Cast<UMediaOutput>(InObject)) == INDEX_NONE)
+		{
+			return;
+		}
+	}
+	
 	FillMediaTree();
 }
 
@@ -1187,17 +1281,41 @@ void SMediaProfileSourcesTreeView::OnMoveMediaItem(TSharedPtr<FMediaTreeItem> In
 	
 	FScopedTransaction MoveMediaItemTransaction(TransactionText);
 	MediaProfile->Modify();
-	
+
+	bool bSuccessful = false;
 	if (bIsMediaSource)
 	{
-		MediaProfile->MoveMediaSource(InTreeItem->Index, InIndex);
+		bSuccessful =MediaProfile->MoveMediaSource(InTreeItem->Index, InIndex);
 	}
 	else
 	{
-		MediaProfile->MoveMediaOutput(InTreeItem->Index, InIndex);
+		bSuccessful = MediaProfile->MoveMediaOutput(InTreeItem->Index, InIndex);
 	}
 
 	FillMediaTree();
+
+	// Update the selection to select the newly moved item
+	if (MediaTreeView.IsValid() && bSuccessful)
+	{
+		MediaTreeView->ClearSelection();
+		if (bIsMediaSource)
+		{
+			MediaTreeView->SetItemSelection(MediaTreeItems[0]->Children[InIndex], true);
+		}
+		else
+		{
+			MediaTreeView->SetItemSelection(MediaTreeItems[1]->Children[InIndex], true);
+		}
+	}
+}
+
+void SMediaProfileSourcesTreeView::OnMediaTreeSelectionChanged(TSharedPtr<FMediaTreeItem> SelectedItem, ESelectInfo::Type SelectionType)
+{
+	TArray<int32> SelectedMediaSources;
+	TArray<int32> SelectedMediaOutputs;
+	GetSelectedMediaIndexes(SelectedMediaSources, SelectedMediaOutputs);
+	
+	OnSelectedMediaItemsChanged.ExecuteIfBound(SelectedMediaSources, SelectedMediaOutputs);
 }
 
 #undef LOCTEXT_NAMESPACE
diff --git a/Engine/Plugins/Media/MediaFrameworkUtilities/Source/MediaFrameworkUtilitiesEditor/Private/AssetEditor/SMediaProfileSourcesTreeView.h b/Engine/Plugins/Media/MediaFrameworkUtilities/Source/MediaFrameworkUtilitiesEditor/Private/AssetEditor/SMediaProfileSourcesTreeView.h
--- a/Engine/Plugins/Media/MediaFrameworkUtilities/Source/MediaFrameworkUtilitiesEditor/Private/AssetEditor/SMediaProfileSourcesTreeView.h
+++ b/Engine/Plugins/Media/MediaFrameworkUtilities/Source/MediaFrameworkUtilitiesEditor/Private/AssetEditor/SMediaProfileSourcesTreeView.h
@@ -17,6 +17,9 @@ class UMediaProfile;
  */
 class SMediaProfileSourcesTreeView : public SCompoundWidget, public FEditorUndoClient
 {
+public:
+	DECLARE_DELEGATE_TwoParams(FOnSelectedMediaItemsChanged, const TArray<int32>& /* SelectedMediaSources */, const TArray<int32>& /* SelectedMediaOutputs */);
+	
 private:
 	/** Struct to store information for each element displayed in the tree view */
 	struct FMediaTreeItem
@@ -44,6 +47,7 @@ private:
 	
 public:
 	SLATE_BEGIN_ARGS(SMediaProfileSourcesTreeView) {}
+		SLATE_EVENT(FOnSelectedMediaItemsChanged, OnSelectedMediaItemsChanged)
 	SLATE_END_ARGS()
 	
 	void Construct(const FArguments& InArgs, UMediaProfile* InMediaProfile);
@@ -83,6 +87,9 @@ private:
 	/** Recreates the tree view's source list from the current media profile's list of sources and outputs */
 	void FillMediaTree();
 
+	/** Converts the selected tree items to a list of indexes of the selected media sources and outputs in the media profile */
+	void GetSelectedMediaIndexes(TArray<int32>& OutSelectedMediaSources, TArray<int32>& OutSelectedMediaOutputs, bool bIncludeGroupIndex = false);
+	
 	// TreeView event handlers
 	void GetTreeItemChildren(FMediaTreeItemPtr InTreeItem, TArray<FMediaTreeItemPtr>& OutChildren) const;
 	TSharedRef<ITableRow> GenerateTreeItemRow(FMediaTreeItemPtr InTreeItem, const TSharedRef<STableViewBase>& InOwnerTableView);
@@ -97,6 +104,8 @@ private:
 	/** Raised when a media tree item is moved */
 	void OnMoveMediaItem(TSharedPtr<FMediaTreeItem> InTreeItem, int32 InIndex);
 	
+	void OnMediaTreeSelectionChanged(TSharedPtr<FMediaTreeItem> SelectedItem, ESelectInfo::Type SelectionType);
+	
 private:
 	TArray<FMediaTreeItemPtr> MediaTreeItems;
 	TSharedPtr<STreeView<FMediaTreeItemPtr>> MediaTreeView;
@@ -104,6 +113,8 @@ private:
 	TWeakObjectPtr<UMediaProfile> MediaProfile;
 
 	TSharedPtr<SMediaProfileSourcesToolbar> Toolbar;
+
+	FOnSelectedMediaItemsChanged OnSelectedMediaItemsChanged;
 	
 	friend class SMediaTreeItemRow;
 	friend class SMediaProfileSourcesToolbar;
diff --git a/Engine/Plugins/Media/MediaIOFramework/Source/MediaIOCore/Private/FileMediaOutput.cpp b/Engine/Plugins/Media/MediaIOFramework/Source/MediaIOCore/Private/FileMediaOutput.cpp
--- a/Engine/Plugins/Media/MediaIOFramework/Source/MediaIOCore/Private/FileMediaOutput.cpp
+++ b/Engine/Plugins/Media/MediaIOFramework/Source/MediaIOCore/Private/FileMediaOutput.cpp
@@ -108,6 +108,14 @@ FString UFileMediaOutput::GetDescriptionString() const
 {
 	return FilePath.Path;
 }
+
+void UFileMediaOutput::GetDetailsPanelInfoElements(TArray<FInfoElement>& OutInfoElements) const
+{
+	OutInfoElements.Add(FInfoElement(
+		NSLOCTEXT("FileMediaOutput", "MediaConfigurationHeader", "Media Configuration"),
+		NSLOCTEXT("FileMediaOutput", "FilePathLabel", "File Path"),
+		FText::FromString(FilePath.Path)));
+}
 #endif //WITH_EDITOR
 
 UMediaCapture* UFileMediaOutput::CreateMediaCaptureImpl()
diff --git a/Engine/Plugins/Media/MediaIOFramework/Source/MediaIOCore/Public/FileMediaOutput.h b/Engine/Plugins/Media/MediaIOFramework/Source/MediaIOCore/Public/FileMediaOutput.h
--- a/Engine/Plugins/Media/MediaIOFramework/Source/MediaIOCore/Public/FileMediaOutput.h
+++ b/Engine/Plugins/Media/MediaIOFramework/Source/MediaIOCore/Public/FileMediaOutput.h
@@ -76,6 +76,7 @@ public:
 
 #if WITH_EDITOR
 	virtual FString GetDescriptionString() const override;
+	virtual void GetDetailsPanelInfoElements(TArray<FInfoElement>& OutInfoElements) const override;
 #endif //WITH_EDITOR
 	
 protected:
diff --git a/Engine/Plugins/Media/MediaIOFramework/Source/MediaIOCore/Public/MediaCapture.h b/Engine/Plugins/Media/MediaIOFramework/Source/MediaIOCore/Public/MediaCapture.h
--- a/Engine/Plugins/Media/MediaIOFramework/Source/MediaIOCore/Public/MediaCapture.h
+++ b/Engine/Plugins/Media/MediaIOFramework/Source/MediaIOCore/Public/MediaCapture.h
@@ -175,6 +175,9 @@ struct FMediaCaptureOptions
 public:
 	UE_API void PostSerialize(const FArchive& Ar);
 
+	/** Gets the string version of ColorConversionSettings */
+	UE_API FString GetColorConversionSettingsString() const { return ColorConversionSettings.ToString(); }
+
 	/** If source texture's size change, auto restart capture to follow source resolution if implementation supports it. */
 	UPROPERTY(BlueprintReadWrite, EditAnywhere, Category = "MediaCapture")
 	bool bAutoRestartOnSourceSizeChange = false;
diff --git a/Engine/Plugins/Media/MediaIOFramework/Source/MediaIOCore/Public/MediaOutput.h b/Engine/Plugins/Media/MediaIOFramework/Source/MediaIOCore/Public/MediaOutput.h
--- a/Engine/Plugins/Media/MediaIOFramework/Source/MediaIOCore/Public/MediaOutput.h
+++ b/Engine/Plugins/Media/MediaIOFramework/Source/MediaIOCore/Public/MediaOutput.h
@@ -91,6 +91,14 @@ public:
 #if WITH_EDITOR
 	/** Gets a string that sums up the media output's description for UI purposes */
 	virtual FString GetDescriptionString() const { return FString(); }
+
+	typedef TTuple<FText, FText, FText> FInfoElement;
+	
+	/**
+	 * Gets a list of text containing the media source's configuration and description for UI purposes
+	 * @param OutInfoElements A list of elements, where an element is constructed from a group, a label, and a value
+	 */
+	virtual void GetDetailsPanelInfoElements(TArray<FInfoElement>& OutInfoElements) const { }
 	
 	DECLARE_MULTICAST_DELEGATE_OneParam(FOnOutputModified, UMediaOutput*);
 	FOnOutputModified& OnOutputModified() { return OnOutputModifiedDelegate; }
diff --git a/Engine/Plugins/Media/NDIMedia/Source/NDIMedia/Private/NDIMediaOutput.cpp b/Engine/Plugins/Media/NDIMedia/Source/NDIMedia/Private/NDIMediaOutput.cpp
--- a/Engine/Plugins/Media/NDIMedia/Source/NDIMedia/Private/NDIMediaOutput.cpp
+++ b/Engine/Plugins/Media/NDIMedia/Source/NDIMedia/Private/NDIMediaOutput.cpp
@@ -67,6 +67,32 @@ FString UNDIMediaOutput::GetDescriptionString() const
 		FString::SanitizeFloat(FrameRate.AsDecimal()),
 		StaticEnum<ENDIMediaOutputPixelFormat>()->GetDisplayNameTextByValue(static_cast<int64>(DesiredPixelFormat)).ToString() });
 }
+
+void UNDIMediaOutput::GetDetailsPanelInfoElements(TArray<FInfoElement>& OutInfoElements) const
+{
+#define LOCTEXT_NAMESPACE "NDIMediaOutput"
+	const FNumberFormattingOptions NumberFormat = FNumberFormattingOptions()
+		.DefaultNoGrouping();
+	
+	const FText ConfigHeader = LOCTEXT("MediaConfigurationHeader", "Media Configuration");
+	
+	const FIntPoint Size = GetRequestedSize();
+	const FText ResolutionText = FText::Format(LOCTEXT("ResolutionFormat", "{0}x{1}"), FText::AsNumber(Size.X, &NumberFormat), FText::AsNumber(Size.Y, &NumberFormat));
+	const FText FramerateText = FText::Format(LOCTEXT("FramerateFormat", "{0} fps"), FText::AsNumber(FrameRate.AsDecimal()));
+	const FText PixelFormatText = StaticEnum<ENDIMediaOutputPixelFormat>()->GetDisplayNameTextByValue(static_cast<int64>(DesiredPixelFormat));
+	
+	OutInfoElements.Add(FInfoElement(ConfigHeader, LOCTEXT("ResolutionLabel", "Resolution"), ResolutionText));
+	OutInfoElements.Add(FInfoElement(ConfigHeader, LOCTEXT("FramerateLabel", "Framerate"), FramerateText));
+	OutInfoElements.Add(FInfoElement(ConfigHeader, LOCTEXT("PixelFormatLabel", "Pixel Format"), PixelFormatText));
+
+	const FText AdditionalHeader = LOCTEXT("AdditionalSettingsHeader", "Additional Settings");
+	const FText DisabledText = LOCTEXT("Disabled", "Disabled");
+	const FText AudioText = FText::Format(LOCTEXT("AudioFormat", "{0} channels, {1} sample rate"),
+		FText::AsNumber(NumOutputAudioChannels), UEnum::GetDisplayValueAsText(AudioSampleRate));
+	
+	OutInfoElements.Add(FInfoElement(AdditionalHeader, LOCTEXT("AudioLabel", "Audio"), bOutputAudio ? AudioText : DisabledText));
+#undef LOCTEXT_NAMESPACE
+}
 #endif //WITH_EDITOR
 
 UMediaCapture* UNDIMediaOutput::CreateMediaCaptureImpl()
diff --git a/Engine/Plugins/Media/NDIMedia/Source/NDIMedia/Public/NDIMediaOutput.h b/Engine/Plugins/Media/NDIMedia/Source/NDIMedia/Public/NDIMediaOutput.h
--- a/Engine/Plugins/Media/NDIMedia/Source/NDIMedia/Public/NDIMediaOutput.h
+++ b/Engine/Plugins/Media/NDIMedia/Source/NDIMedia/Public/NDIMediaOutput.h
@@ -146,6 +146,7 @@ public:
 
 #if WITH_EDITOR
 	virtual FString GetDescriptionString() const override;
+	virtual void GetDetailsPanelInfoElements(TArray<FInfoElement>& OutInfoElements) const override;
 #endif //WITH_EDITOR
 
 protected:
diff --git a/Engine/Plugins/VirtualProduction/Rivermax/RivermaxMedia/Source/RivermaxMedia/Private/RivermaxMediaOutput.cpp b/Engine/Plugins/VirtualProduction/Rivermax/RivermaxMedia/Source/RivermaxMedia/Private/RivermaxMediaOutput.cpp
--- a/Engine/Plugins/VirtualProduction/Rivermax/RivermaxMedia/Source/RivermaxMedia/Private/RivermaxMediaOutput.cpp
+++ b/Engine/Plugins/VirtualProduction/Rivermax/RivermaxMedia/Source/RivermaxMedia/Private/RivermaxMediaOutput.cpp
@@ -67,6 +67,34 @@ FString URivermaxMediaOutput::GetDescriptionString() const
 		FString::SanitizeFloat(VideoStream.FrameRate.AsDecimal()),
 		StaticEnum<ERivermaxMediaOutputPixelFormat>()->GetDisplayNameTextByValue(static_cast<int64>(VideoStream.PixelFormat)).ToString() });
 }
+
+void URivermaxMediaOutput::GetDetailsPanelInfoElements(TArray<FInfoElement>& OutInfoElements) const
+{
+#define LOCTEXT_NAMESPACE "RivermaxMediaOutput"
+	const FNumberFormattingOptions NumberFormat = FNumberFormattingOptions()
+		.DefaultNoGrouping();
+	
+	const FText ConfigHeader = LOCTEXT("MediaConfigurationHeader", "Media Configuration");
+
+	const FIntPoint Size = GetRequestedSize();
+	const FText ResolutionText = FText::Format(LOCTEXT("ResolutionFormat", "{0}x{1}"), FText::AsNumber(Size.X, &NumberFormat), FText::AsNumber(Size.Y, &NumberFormat));
+	const FText FramerateText = FText::Format(LOCTEXT("FramerateFormat", "{0} fps"), FText::AsNumber(VideoStream.FrameRate.AsDecimal()));
+	const FText PixelFormatText = StaticEnum<ERivermaxMediaSourcePixelFormat>()->GetDisplayNameTextByValue(static_cast<int64>(VideoStream.PixelFormat));
+	
+	OutInfoElements.Add(FInfoElement(ConfigHeader, LOCTEXT("ResolutionLabel", "Resolution"), ResolutionText));
+	OutInfoElements.Add(FInfoElement(ConfigHeader, LOCTEXT("FramerateLabel", "Framerate"), FramerateText));
+	OutInfoElements.Add(FInfoElement(ConfigHeader, LOCTEXT("PixelFormatLabel", "Pixel Format"), PixelFormatText));
+	OutInfoElements.Add(FInfoElement(ConfigHeader, LOCTEXT("StreamAddressLabel", "Stream Address"), FText::FromString(StreamAddress)));
+	OutInfoElements.Add(FInfoElement(ConfigHeader, LOCTEXT("PortLabel", "Port"), FText::AsNumber(Port, &NumberFormat)));
+
+	const FText VideoSettingsHeader = LOCTEXT("VideoSettingsHeader", "Video Settings");
+	const FText EnabledText = LOCTEXT("Enabled", "Enabled");
+	const FText DisabledText = LOCTEXT("Disabled", "Disabled");
+
+	OutInfoElements.Add(FInfoElement(VideoSettingsHeader, LOCTEXT("GPUDirectLabel", "GPUDirect"), VideoStream.bUseGPUDirect ? EnabledText : DisabledText));
+	
+#undef LOCTEXT_NAMESPACE
+}
 #endif //WITH_EDITOR
 
 void URivermaxMediaOutput::ExportSDP(const FString& InPath)
diff --git a/Engine/Plugins/VirtualProduction/Rivermax/RivermaxMedia/Source/RivermaxMedia/Private/RivermaxMediaSource.cpp b/Engine/Plugins/VirtualProduction/Rivermax/RivermaxMedia/Source/RivermaxMedia/Private/RivermaxMediaSource.cpp
--- a/Engine/Plugins/VirtualProduction/Rivermax/RivermaxMedia/Source/RivermaxMedia/Private/RivermaxMediaSource.cpp
+++ b/Engine/Plugins/VirtualProduction/Rivermax/RivermaxMedia/Source/RivermaxMedia/Private/RivermaxMediaSource.cpp
@@ -173,4 +173,32 @@ FString URivermaxMediaSource::GetDescriptionString() const
 		FString::SanitizeFloat(FrameRate.AsDecimal()),
 		StaticEnum<ERivermaxMediaSourcePixelFormat>()->GetDisplayNameTextByValue(static_cast<int64>(PixelFormat)).ToString() });
 }
+
+void URivermaxMediaSource::GetDetailsPanelInfoElements(TArray<FInfoElement>& OutInfoElements) const
+{
+#define LOCTEXT_NAMESPACE "RivermaxMediaSource"
+	const FNumberFormattingOptions NumberFormat = FNumberFormattingOptions()
+		.DefaultNoGrouping();
+	
+	const FText ConfigHeader = LOCTEXT("MediaConfigurationHeader", "Media Configuration");
+
+	const FText ResolutionText = FText::Format(LOCTEXT("ResolutionFormat", "{0}x{1}"), FText::AsNumber(Resolution.X, &NumberFormat), FText::AsNumber(Resolution.Y, &NumberFormat));
+	const FText FramerateText = FText::Format(LOCTEXT("FramerateFormat", "{0} fps"), FText::AsNumber(FrameRate.AsDecimal()));
+	const FText PixelFormatText = StaticEnum<ERivermaxMediaSourcePixelFormat>()->GetDisplayNameTextByValue(static_cast<int64>(PixelFormat));
+	
+	OutInfoElements.Add(FInfoElement(ConfigHeader, LOCTEXT("ResolutionLabel", "Resolution"), ResolutionText));
+	OutInfoElements.Add(FInfoElement(ConfigHeader, LOCTEXT("FramerateLabel", "Framerate"), FramerateText));
+	OutInfoElements.Add(FInfoElement(ConfigHeader, LOCTEXT("PixelFormatLabel", "Pixel Format"), PixelFormatText));
+	OutInfoElements.Add(FInfoElement(ConfigHeader, LOCTEXT("StreamAddressLabel", "Stream Address"), FText::FromString(StreamAddress)));
+	OutInfoElements.Add(FInfoElement(ConfigHeader, LOCTEXT("PortLabel", "Port"), FText::AsNumber(Port, &NumberFormat)));
+
+	const FText VideoSettingsHeader = LOCTEXT("VideoSettingsHeader", "Video Settings");
+	const FText EnabledText = LOCTEXT("Enabled", "Enabled");
+	const FText DisabledText = LOCTEXT("Disabled", "Disabled");
+
+	OutInfoElements.Add(FInfoElement(VideoSettingsHeader, LOCTEXT("GPUDirectLabel", "GPUDirect"), bUseGPUDirect ? EnabledText : DisabledText));
+	OutInfoElements.Add(FInfoElement(VideoSettingsHeader, LOCTEXT("ColorConversionLabel", "Color Conversion"), FText::FromString(ColorConversionSettings.ToString())));
+	
+#undef LOCTEXT_NAMESPACE
+}
 #endif //WITH_EDITOR
diff --git a/Engine/Plugins/VirtualProduction/Rivermax/RivermaxMedia/Source/RivermaxMedia/Public/RivermaxMediaOutput.h b/Engine/Plugins/VirtualProduction/Rivermax/RivermaxMedia/Source/RivermaxMedia/Public/RivermaxMediaOutput.h
--- a/Engine/Plugins/VirtualProduction/Rivermax/RivermaxMedia/Source/RivermaxMedia/Public/RivermaxMediaOutput.h
+++ b/Engine/Plugins/VirtualProduction/Rivermax/RivermaxMedia/Source/RivermaxMedia/Public/RivermaxMediaOutput.h
@@ -127,6 +127,7 @@ public:
 	
 #if WITH_EDITOR
 	virtual FString GetDescriptionString() const override;
+	virtual void GetDetailsPanelInfoElements(TArray<FInfoElement>& OutInfoElements) const override;
 #endif //WITH_EDITOR
 
 public:
diff --git a/Engine/Plugins/VirtualProduction/Rivermax/RivermaxMedia/Source/RivermaxMedia/Public/RivermaxMediaSource.h b/Engine/Plugins/VirtualProduction/Rivermax/RivermaxMedia/Source/RivermaxMedia/Public/RivermaxMediaSource.h
--- a/Engine/Plugins/VirtualProduction/Rivermax/RivermaxMedia/Source/RivermaxMedia/Public/RivermaxMediaSource.h
+++ b/Engine/Plugins/VirtualProduction/Rivermax/RivermaxMedia/Source/RivermaxMedia/Public/RivermaxMediaSource.h
@@ -130,6 +130,7 @@ public:
 
 #if WITH_EDITOR
 	virtual FString GetDescriptionString() const override;
+	virtual void GetDetailsPanelInfoElements(TArray<FInfoElement>& OutInfoElements) const override;
 #endif //WITH_EDITOR
 	
 	//~ End UMediaSource interface
diff --git a/Engine/Plugins/VirtualProduction/Rivermax/RivermaxMedia/Source/RivermaxMedia/RivermaxMedia.build.cs b/Engine/Plugins/VirtualProduction/Rivermax/RivermaxMedia/Source/RivermaxMedia/RivermaxMedia.build.cs
--- a/Engine/Plugins/VirtualProduction/Rivermax/RivermaxMedia/Source/RivermaxMedia/RivermaxMedia.build.cs
+++ b/Engine/Plugins/VirtualProduction/Rivermax/RivermaxMedia/Source/RivermaxMedia/RivermaxMedia.build.cs
@@ -20,6 +20,7 @@ public class RivermaxMedia : ModuleRules
 		PrivateDependencyModuleNames.AddRange(
 			new string[] {
 				"Engine",
+				"OpenColorIO",
 				"RenderCore",
 				"RHI",
 				"RivermaxRendering"
diff --git a/Engine/Source/Runtime/MediaAssets/Private/Assets/FileMediaSource.cpp b/Engine/Source/Runtime/MediaAssets/Private/Assets/FileMediaSource.cpp
--- a/Engine/Source/Runtime/MediaAssets/Private/Assets/FileMediaSource.cpp
+++ b/Engine/Source/Runtime/MediaAssets/Private/Assets/FileMediaSource.cpp
@@ -158,3 +158,13 @@ bool UFileMediaSource::Validate() const
 	return FPaths::FileExists(ResolvedFullPath);
 }
 
+#if WITH_EDITOR
+void UFileMediaSource::GetDetailsPanelInfoElements(TArray<FInfoElement>& OutInfoElements) const
+{
+	OutInfoElements.Add(FInfoElement(
+		NSLOCTEXT("FileMediaSource", "MediaConfigurationHeader", "Media Configuration"),
+		NSLOCTEXT("FileMediaSource", "FilePathLabel", "File Path"),
+		FText::FromString(GetFullPath())));
+}
+#endif
+
diff --git a/Engine/Source/Runtime/MediaAssets/Private/Assets/MediaSource.cpp b/Engine/Source/Runtime/MediaAssets/Private/Assets/MediaSource.cpp
--- a/Engine/Source/Runtime/MediaAssets/Private/Assets/MediaSource.cpp
+++ b/Engine/Source/Runtime/MediaAssets/Private/Assets/MediaSource.cpp
@@ -39,6 +39,13 @@ void UMediaSource::GenerateThumbnail()
 	}
 }
 
+void UMediaSource::GetDetailsPanelInfoElements(TArray<FInfoElement>& OutInfoElements) const
+{
+	OutInfoElements.Add(FInfoElement(
+		NSLOCTEXT("MediaSource", "MediaConfigurationHeader", "Media Configuration"),
+		NSLOCTEXT("MediaSource", "UrlLabel", "Url"),
+		FText::FromString(GetUrl())));
+}
 #endif // WITH_EDITOR
 
 void UMediaSource::RegisterSpawnFromFileExtension(const FString& Extension,
diff --git a/Engine/Source/Runtime/MediaAssets/Private/Assets/PlatformMediaSource.cpp b/Engine/Source/Runtime/MediaAssets/Private/Assets/PlatformMediaSource.cpp
--- a/Engine/Source/Runtime/MediaAssets/Private/Assets/PlatformMediaSource.cpp
+++ b/Engine/Source/Runtime/MediaAssets/Private/Assets/PlatformMediaSource.cpp
@@ -202,6 +202,24 @@ FString UPlatformMediaSource::GetDescriptionString() const
 	UMediaSource* PlatformMediaSource = GetMediaSource();
 	return (PlatformMediaSource != nullptr) ? PlatformMediaSource->GetDescriptionString() : FString();
 }
+
+void UPlatformMediaSource::GetDetailsPanelInfoElements(TArray<FInfoElement>& OutInfoElements) const
+{
+	// Guard against reentrant calls.
+	static bool bIsGettingInfoElements = false;
+	if (bIsGettingInfoElements)
+	{
+		UE_LOG(LogMediaAssets, Warning, TEXT("UPlatformMediaSource::GetDetailsPanelInfoElements - Reentrant calls are not supported. Asset: %s"), *GetPathName());
+		return;
+	}
+	TGuardValue<bool> Guard(bIsGettingInfoElements, true);
+
+	UMediaSource* PlatformMediaSource = GetMediaSource();
+	if (PlatformMediaSource != nullptr)
+	{
+		PlatformMediaSource->GetDetailsPanelInfoElements(OutInfoElements);
+	}
+}
 #endif //WITH_EDITOR
 
 /* UPlatformMediaSource implementation
diff --git a/Engine/Source/Runtime/MediaAssets/Public/FileMediaSource.h b/Engine/Source/Runtime/MediaAssets/Public/FileMediaSource.h
--- a/Engine/Source/Runtime/MediaAssets/Public/FileMediaSource.h
+++ b/Engine/Source/Runtime/MediaAssets/Public/FileMediaSource.h
@@ -89,6 +89,9 @@ public:
 	MEDIAASSETS_API virtual FString GetUrl() const override;
 	MEDIAASSETS_API virtual bool Validate() const override;
 
+#if WITH_EDITOR
+	MEDIAASSETS_API virtual void GetDetailsPanelInfoElements(TArray<FInfoElement>& OutInfoElements) const override;
+#endif
 public:
 
 	/** Name of the PrecacheFile media option. */
diff --git a/Engine/Source/Runtime/MediaAssets/Public/MediaSource.h b/Engine/Source/Runtime/MediaAssets/Public/MediaSource.h
--- a/Engine/Source/Runtime/MediaAssets/Public/MediaSource.h
+++ b/Engine/Source/Runtime/MediaAssets/Public/MediaSource.h
@@ -83,7 +83,14 @@ public:
 
 	/** Gets a string that sums up the media source's description for UI purposes */
 	virtual FString GetDescriptionString() const { return GetUrl(); }
+
+	typedef TTuple<FText, FText, FText> FInfoElement;
 	
+	/**
+	 * Gets a list of text containing the media source's configuration and description for UI purposes
+	 * @param OutInfoElements A list of elements, where an element is constructed from a group, a label, and a value
+	 */
+	MEDIAASSETS_API virtual void GetDetailsPanelInfoElements(TArray<FInfoElement>& OutInfoElements) const;
 #endif // WITH_EDITOR
 
 	/**
diff --git a/Engine/Source/Runtime/MediaAssets/Public/PlatformMediaSource.h b/Engine/Source/Runtime/MediaAssets/Public/PlatformMediaSource.h
--- a/Engine/Source/Runtime/MediaAssets/Public/PlatformMediaSource.h
+++ b/Engine/Source/Runtime/MediaAssets/Public/PlatformMediaSource.h
@@ -57,6 +57,7 @@ public:
 
 #if WITH_EDITOR
 	virtual FString GetDescriptionString() const override;
+	virtual void GetDetailsPanelInfoElements(TArray<FInfoElement>& OutInfoElements) const override;
 #endif //WITH_EDITOR
 	
 public:
