Media Profile Editor: Added viewports tab which displays the selected media sources and media outputs to the user. Media source streams are managed by a media profile playback manager, which creates media players as needed and outputs them to default media textures stored in the plugin's content. Media output captures are managed by the media profile editor itself using media capture objects and viewport clients and widgets.

#virtualized

diff --git a/Engine/Plugins/Media/MediaFrameworkUtilities/Source/MediaFrameworkUtilities/Private/Profile/MediaProfile.cpp b/Engine/Plugins/Media/MediaFrameworkUtilities/Source/MediaFrameworkUtilities/Private/Profile/MediaProfile.cpp
--- a/Engine/Plugins/Media/MediaFrameworkUtilities/Source/MediaFrameworkUtilities/Private/Profile/MediaProfile.cpp
+++ b/Engine/Plugins/Media/MediaFrameworkUtilities/Source/MediaFrameworkUtilities/Private/Profile/MediaProfile.cpp
@@ -12,6 +12,7 @@
 #include "MediaAssets/ProxyMediaOutput.h"
 #include "MediaAssets/ProxyMediaSource.h"
 #include "Profile/IMediaProfileManager.h"
+#include "Profile/MediaProfilePlaybackManager.h"
 
 #if WITH_EDITOR
 #include "EngineAnalytics.h"
@@ -43,6 +44,11 @@ namespace MediaProfileAnalytics
 }
 #endif
 
+UMediaProfile::UMediaProfile(const FObjectInitializer& ObjectInitializer)
+{
+	PlaybackManager = CreateDefaultSubobject<UMediaProfilePlaybackManager>(TEXT("PlaybackManager"), true);
+}
+
 UMediaSource* UMediaProfile::GetMediaSource(int32 Index) const
 {
 	if (MediaSources.IsValidIndex(Index))
diff --git a/Engine/Plugins/Media/MediaFrameworkUtilities/Source/MediaFrameworkUtilities/Private/Profile/MediaProfilePlaybackManager.cpp b/Engine/Plugins/Media/MediaFrameworkUtilities/Source/MediaFrameworkUtilities/Private/Profile/MediaProfilePlaybackManager.cpp
new file mode 100644
--- /dev/null
+++ b/Engine/Plugins/Media/MediaFrameworkUtilities/Source/MediaFrameworkUtilities/Private/Profile/MediaProfilePlaybackManager.cpp
@@ -0,0 +1,180 @@
+﻿// Copyright Epic Games, Inc. All Rights Reserved.
+
+
+#include "Profile/MediaProfilePlaybackManager.h"
+
+#include "MediaPlayer.h"
+#include "MediaSource.h"
+#include "MediaTexture.h"
+#include "Profile/MediaProfile.h"
+
+namespace MediaProfilePlaybackManager
+{
+	const FString DefaultMediaTexturesPath = TEXT("/MediaFrameworkUtilities/DefaultMediaTextures");
+	const FString DefaultMediaTextureName = TEXT("MediaProfileTexture");
+}
+
+UMediaProfilePlaybackManager::UMediaProfilePlaybackManager(const FObjectInitializer& ObjectInitializer)
+	: Super(ObjectInitializer)
+{
+	constexpr int32 NumMediaTextures = 16;
+	for (int32 TextureIndex = 1; TextureIndex <= NumMediaTextures; ++TextureIndex)
+	{
+		const FString IndexStr = TextureIndex > 9 ? FString::FromInt(TextureIndex) : TEXT("0") + FString::FromInt(TextureIndex);
+		const FString TextureAssetName = FString::Format(TEXT("{0}_{1}"), {
+			MediaProfilePlaybackManager::DefaultMediaTextureName,
+			IndexStr });
+		
+		const FString Path = FString::Format(TEXT("{0}/{1}.{1}"), {
+			MediaProfilePlaybackManager::DefaultMediaTexturesPath,
+			TextureAssetName });
+
+		MediaSourceTextures.Add(TSoftObjectPtr<UMediaTexture>(FSoftObjectPath(Path)));
+	}
+}
+
+UMediaTexture* UMediaProfilePlaybackManager::GetSourceMediaTextureFromIndex(int32 InMediaSourceIndex) const
+{
+	if (!MediaSourceTextures.IsValidIndex(InMediaSourceIndex))
+	{
+		return nullptr;
+	}
+
+	UMediaProfile* MediaProfile = GetOwningMediaProfile();
+	if (!MediaProfile)
+	{
+		return nullptr;
+	}
+	
+	if (InMediaSourceIndex >= MediaProfile->NumMediaSources())
+	{
+		return nullptr;	
+	}
+
+	return MediaSourceTextures[InMediaSourceIndex].LoadSynchronous();
+}
+
+UMediaTexture* UMediaProfilePlaybackManager::GetSourceMediaTexture(UMediaSource* InMediaSource) const
+{
+	UMediaProfile* MediaProfile = GetOwningMediaProfile();
+	if (!MediaProfile)
+	{
+		return nullptr;
+	}
+	
+	int32 MediaSourceIndex = MediaProfile->FindMediaSourceIndex(InMediaSource);
+	return GetSourceMediaTextureFromIndex(MediaSourceIndex);
+}
+
+UMediaTexture* UMediaProfilePlaybackManager::OpenSourceFromIndex(int32 InMediaSourceIndex)
+{
+	if (!MediaSourceTextures.IsValidIndex(InMediaSourceIndex))
+	{
+		return nullptr;
+	}
+	
+	UMediaProfile* MediaProfile = GetOwningMediaProfile();
+	if (!MediaProfile)
+	{
+		return nullptr;
+	}
+	
+	if (InMediaSourceIndex < 0 || InMediaSourceIndex >= MediaProfile->NumMediaSources())
+	{
+		return nullptr;
+	}
+
+	UMediaSource* MediaSource = MediaProfile->GetMediaSource(InMediaSourceIndex);
+	UMediaPlayer* MediaSourcePlayer = FindOrCreateMediaPlayerForIndex(InMediaSourceIndex);	
+	if (MediaSourcePlayer->IsPlaying())
+	{
+		// If the media source is already being played, simply return the media texture
+		if (MediaSourcePlayer->GetUrl() == MediaSource->GetUrl())
+		{
+			return MediaSourceTextures[InMediaSourceIndex].LoadSynchronous();
+		}
+	}
+
+	FMediaPlayerOptions Options;
+	Options.PlayOnOpen = EMediaPlayerOptionBooleanOverride::Enabled;
+	Options.Loop = EMediaPlayerOptionBooleanOverride::Enabled;
+	
+	if (MediaSourcePlayer->OpenSourceWithOptions(MediaSource, Options))
+	{
+		return MediaSourceTextures[InMediaSourceIndex].LoadSynchronous();
+	}
+	
+	// If the media player was unable to open the source, return nullptr to indicate that the media texture will not be displaying the media source
+	return nullptr;
+}
+
+UMediaTexture* UMediaProfilePlaybackManager::OpenSource(UMediaSource* InMediaSource)
+{
+	UMediaProfile* MediaProfile = GetOwningMediaProfile();
+	if (!MediaProfile)
+	{
+		return nullptr;
+	}
+	
+	int32 MediaSourceIndex = MediaProfile->FindMediaSourceIndex(InMediaSource);
+	return OpenSourceFromIndex(MediaSourceIndex);
+}
+
+void UMediaProfilePlaybackManager::CloseSourceFromIndex(int32 InMediaSourceIndex, bool bDestroyMediaPlayer)
+{
+	if (!MediaSourcePlayers.Contains(InMediaSourceIndex))
+	{
+		return;
+	}
+
+	MediaSourcePlayers[InMediaSourceIndex]->Close();
+
+	if (bDestroyMediaPlayer)
+	{
+		MediaSourcePlayers.Remove(InMediaSourceIndex);
+	}
+}
+
+void UMediaProfilePlaybackManager::CloseSource(UMediaSource* InMediaSource, bool bDestroyMediaPlayer)
+{
+	UMediaProfile* MediaProfile = GetOwningMediaProfile();
+	if (!MediaProfile)
+	{
+		return;
+	}
+	
+	int32 MediaSourceIndex = MediaProfile->FindMediaSourceIndex(InMediaSource);
+	CloseSourceFromIndex(MediaSourceIndex, bDestroyMediaPlayer);
+}
+
+UMediaPlayer* UMediaProfilePlaybackManager::FindOrCreateMediaPlayerForIndex(int32 InMediaSourceIndex)
+{
+	UMediaProfile* MediaProfile = GetOwningMediaProfile();
+	if (!MediaProfile)
+	{
+		return nullptr;
+	}
+	
+	if (InMediaSourceIndex < 0 || InMediaSourceIndex >= MediaProfile->NumMediaSources())
+	{
+		return nullptr;
+	}
+
+	if (!MediaSourcePlayers.Contains(InMediaSourceIndex))
+	{
+		UMediaPlayer* NewMediaPlayer = NewObject<UMediaPlayer>(this);
+		if (UMediaTexture* MediaTexture = GetSourceMediaTextureFromIndex(InMediaSourceIndex))
+		{
+			MediaTexture->SetMediaPlayer(NewMediaPlayer);
+		}
+		
+		MediaSourcePlayers.Add(InMediaSourceIndex, NewMediaPlayer);
+	}
+
+	return MediaSourcePlayers[InMediaSourceIndex];
+}
+
+UMediaProfile* UMediaProfilePlaybackManager::GetOwningMediaProfile() const
+{
+	return GetTypedOuter<UMediaProfile>();
+}
diff --git a/Engine/Plugins/Media/MediaFrameworkUtilities/Source/MediaFrameworkUtilities/Public/Profile/MediaProfile.h b/Engine/Plugins/Media/MediaFrameworkUtilities/Source/MediaFrameworkUtilities/Public/Profile/MediaProfile.h
--- a/Engine/Plugins/Media/MediaFrameworkUtilities/Source/MediaFrameworkUtilities/Public/Profile/MediaProfile.h
+++ b/Engine/Plugins/Media/MediaFrameworkUtilities/Source/MediaFrameworkUtilities/Public/Profile/MediaProfile.h
@@ -8,6 +8,7 @@
 
 #include "MediaProfile.generated.h"
 
+class UMediaProfilePlaybackManager;
 class UEngineCustomTimeStep;
 class UMediaOutput;
 class UMediaSource;
@@ -21,6 +22,9 @@ class MEDIAFRAMEWORKUTILITIES_API UMediaProfile : public UObject
 {
 	GENERATED_BODY()
 
+public:
+	UMediaProfile(const FObjectInitializer& ObjectInitializer);
+
 protected:
 
 	/** Media sources. */
@@ -222,18 +226,14 @@ public:
 	 */
 	virtual void Reset();
 
-
-	/**
-	 * Apply the media profile as the current profile.
-	 * Will change the engine's timecode provider & custom time step and redirect the media profile source/output proxy for the correct media source/output.
-	 */
-	bool IsMediaSourceAffectedByProfile(UMediaSource* InMediaSource);
-
 	/**
 	 * Update the number of sources and outputs to the number to proxies.
 	 */
 	void FixNumSourcesAndOutputs();
 
+	/** Gets the playback manager to use for playback of this profile's media sources and outputs */
+	UMediaProfilePlaybackManager* GetPlaybackManager() const { return PlaybackManager; }
+	
 private:
 
 	void ResetTimecodeProvider();
@@ -255,4 +255,8 @@ private:
 	TObjectPtr<UEngineCustomTimeStep> AppliedCustomTimeStep;
 	UPROPERTY(Transient)
 	TObjectPtr<UEngineCustomTimeStep> PreviousCustomTimeStep;
+
+	/** Playback manager with maintains any necessary objects to support playback of the profile's media sources and outputs */
+	UPROPERTY(Transient, Instanced)
+	TObjectPtr<UMediaProfilePlaybackManager> PlaybackManager;
 };
diff --git a/Engine/Plugins/Media/MediaFrameworkUtilities/Source/MediaFrameworkUtilities/Public/Profile/MediaProfilePlaybackManager.h b/Engine/Plugins/Media/MediaFrameworkUtilities/Source/MediaFrameworkUtilities/Public/Profile/MediaProfilePlaybackManager.h
new file mode 100644
--- /dev/null
+++ b/Engine/Plugins/Media/MediaFrameworkUtilities/Source/MediaFrameworkUtilities/Public/Profile/MediaProfilePlaybackManager.h
@@ -0,0 +1,85 @@
+﻿// Copyright Epic Games, Inc. All Rights Reserved.
+
+#pragma once
+
+#include "UObject/Object.h"
+
+#include "MediaProfilePlaybackManager.generated.h"
+
+class UMediaSource;
+class UMediaTexture;
+class UMediaPlayer;
+class UMediaProfile;
+
+/** Manages playback of media items within a media profile */
+UCLASS()
+class MEDIAFRAMEWORKUTILITIES_API UMediaProfilePlaybackManager : public UObject
+{
+	GENERATED_BODY()
+
+public:
+	UMediaProfilePlaybackManager(const FObjectInitializer& ObjectInitializer);
+	
+public:
+	/**
+	 * Gets the media texture that is outputting the specified media source
+	 * @param InMediaSourceIndex The index of the media source within the media profile
+	 * @return The media texture for the specified media source, or null if no media source exists at the specified index
+	 */
+	UMediaTexture* GetSourceMediaTextureFromIndex(int32 InMediaSourceIndex) const;
+
+	/**
+	 * Gets the media texture that is outputting the specified media source
+	 * @param InMediaSource The media source within the media profile
+	 * @return The media texture for the specified media source, or null if the media source is not a valid member of the media profile
+	 */
+	UMediaTexture* GetSourceMediaTexture(UMediaSource* InMediaSource) const;
+
+	/**
+	 * Opens and begins playing the specified media source 
+	 * @param InMediaSourceIndex The index of the media source within the media profile
+	 * @return The media texture that is outputting the media source, or null if no media source exists at the specified index
+	 */
+	UMediaTexture* OpenSourceFromIndex(int32 InMediaSourceIndex);
+
+	/**
+	 * Opens and begins playing the specified media source 
+	 * @param InMediaSource The media source within the media profile
+	 * @return The media texture that is outputting the media source, or null if the media source is not a valid member of the media profile
+	 */
+	UMediaTexture* OpenSource(UMediaSource* InMediaSource);
+
+	/**
+	 * Closes the specified media source, and optionally destroys the media player that was playing the source
+	 * @param InMediaSourceIndex The index of the media source within the media profile
+	 * @param bDestroyMediaPlayer Optional parameter to destroy the media player playing the source or not
+	 */
+	void CloseSourceFromIndex(int32 InMediaSourceIndex, bool bDestroyMediaPlayer = false);
+	
+	/**
+	 * Closes the specified media source, and optionally destroys the media player that was playing the source
+	 * @param InMediaSource The media source within the media profile
+	 * @param bDestroyMediaPlayer Optional parameter to destroy the media player playing the source or not
+	 */
+	void CloseSource(UMediaSource* InMediaSource, bool bDestroyMediaPlayer = false);
+
+private:
+	/**
+	 * Attempts to find an existing media player for the specified source, and creates one if an existing one is not found
+	 * @param InMediaSourceIndex The index of the media source within the media profile
+	 * @return The media player for the media source, or null if no media source exists at the specified index
+	 */
+	UMediaPlayer* FindOrCreateMediaPlayerForIndex(int32 InMediaSourceIndex);
+
+	/** Gets the media profile that owns this playback manager */
+	UMediaProfile* GetOwningMediaProfile() const;
+	
+private:
+	/** A map of media players, where the hash key is the index of the media source the player is playing */
+	UPROPERTY(Transient)
+	TMap<int32, TObjectPtr<UMediaPlayer>> MediaSourcePlayers;
+
+	/** A list of media textures to use to output the corresponding media sources */
+	UPROPERTY(Transient)
+	TArray<TSoftObjectPtr<UMediaTexture>> MediaSourceTextures;
+};
diff --git a/Engine/Plugins/Media/MediaFrameworkUtilities/Source/MediaFrameworkUtilitiesEditor/MediaFrameworkUtilitiesEditor.Build.cs b/Engine/Plugins/Media/MediaFrameworkUtilities/Source/MediaFrameworkUtilitiesEditor/MediaFrameworkUtilitiesEditor.Build.cs
--- a/Engine/Plugins/Media/MediaFrameworkUtilities/Source/MediaFrameworkUtilitiesEditor/MediaFrameworkUtilitiesEditor.Build.cs
+++ b/Engine/Plugins/Media/MediaFrameworkUtilities/Source/MediaFrameworkUtilitiesEditor/MediaFrameworkUtilitiesEditor.Build.cs
@@ -13,6 +13,7 @@ namespace UnrealBuildTool.Rules
 					"AssetRegistry",
 					"AssetTools",
 					"ClassViewer",
+					"CommonMenuExtensions",
 					"Core",
 					"CoreUObject",
 					"EditorFramework",
diff --git a/Engine/Plugins/Media/MediaFrameworkUtilities/Source/MediaFrameworkUtilitiesEditor/Private/AssetEditor/MediaProfileEditor.cpp b/Engine/Plugins/Media/MediaFrameworkUtilities/Source/MediaFrameworkUtilitiesEditor/Private/AssetEditor/MediaProfileEditor.cpp
--- a/Engine/Plugins/Media/MediaFrameworkUtilities/Source/MediaFrameworkUtilitiesEditor/Private/AssetEditor/MediaProfileEditor.cpp
+++ b/Engine/Plugins/Media/MediaFrameworkUtilities/Source/MediaFrameworkUtilitiesEditor/Private/AssetEditor/MediaProfileEditor.cpp
@@ -8,13 +8,17 @@
 #include "PropertyEditorModule.h"
 #include "SMediaProfileDetailsPanel.h"
 #include "SMediaProfileSourcesTreeView.h"
+#include "SMediaProfileViewport.h"
 #include "CaptureTab/SMediaFrameworkCapture.h"
+#include "Framework/Application/SlateApplication.h"
 #include "GameFramework/WorldSettings.h"
-#include "Misc/TransactionObjectEvent.h"
 #include "Modules/ModuleManager.h"
 #include "Profile/IMediaProfileManager.h"
 #include "Profile/MediaProfile.h"
+#include "Profile/MediaProfilePlaybackManager.h"
+#include "Profile/MediaProfileSettings.h"
 #include "Slate/SceneViewport.h"
+#include "Widgets/SViewport.h"
 #include "Widgets/Docking/SDockTab.h"
 
 #define LOCTEXT_NAMESPACE "MediaProfileEditor"
@@ -140,6 +144,15 @@ private:
 	EViewModeIndex ViewModeIndex;
 };
 
+FMediaProfileEditor::FViewportCapture::~FViewportCapture()
+{
+	ViewportClient->Viewport = nullptr;
+
+	ViewportClient.Reset();
+	Viewport.Reset();
+	ViewportWidget.Reset();
+}
+
 TSharedRef<FMediaProfileEditor> FMediaProfileEditor::CreateMediaProfileEditor(const EToolkitMode::Type Mode, const TSharedPtr<IToolkitHost>& InitToolkitHost, UMediaProfile* InMediaProfile)
 {
 	TSharedRef<FMediaProfileEditor> NewEditor = MakeShared<FMediaProfileEditor>();
@@ -203,6 +216,9 @@ void FMediaProfileEditor::Initialize(const EToolkitMode::Type Mode, const TShare
 	InitAssetEditor(Mode, InitToolkitHost, AppName, StandaloneDefaultLayout, bCreateDefaultStandaloneMenu, bCreateDefaultToolbar, InMediaProfile);
 
 	// TODO: Not sure this is the best place for this, but for now, when we open a media profile, automatically make it the active media profile
+	// First, update the media profile settings config to match the media profile's source and output count
+	GetMutableDefault<UMediaProfileSettings>()->FillDefaultMediaSourceProxies(InMediaProfile->NumMediaSources());
+	GetMutableDefault<UMediaProfileSettings>()->FillDefaultMediaOutputProxies(InMediaProfile->NumMediaOutputs());
 	IMediaProfileManager::Get().SetCurrentMediaProfile(InMediaProfile);
 }
 
@@ -283,6 +299,32 @@ FLinearColor FMediaProfileEditor::GetWorldCentricTabColorScale() const
 	return FLinearColor::White;
 }
 
+TSharedPtr<FMediaProfileEditor::FViewportCapture> FMediaProfileEditor::FindOrCreateViewportCapture(UMediaOutput* InMediaOutput)
+{
+	UMediaFrameworkWorldSettingsAssetUserData* CaptureSettings = FindOrAddMediaFrameworkCaptureSettings();
+	if (!CaptureSettings)
+	{
+		return nullptr;
+	}
+
+	if (!InMediaOutput)
+	{
+		return nullptr;
+	}
+
+	int32 ViewportOutputInfoIndex = CaptureSettings->ViewportCaptures.IndexOfByPredicate([InMediaOutput](const FMediaFrameworkCaptureCameraViewportCameraOutputInfo& ViewportOutputInfo)
+	{
+		return ViewportOutputInfo.MediaOutput == InMediaOutput;
+	});
+
+	if (ViewportOutputInfoIndex == INDEX_NONE)
+	{
+		return nullptr;
+	}
+
+	return FindOrCreateViewportCaptureImpl(ViewportOutputInfoIndex);
+}
+
 bool FMediaProfileEditor::IsMediaOutputCapturing(UMediaOutput* InMediaOutput) const
 {
 	return LiveCaptures.ContainsByPredicate([InMediaOutput](const FLiveMediaCapture& LiveCapture)
@@ -351,8 +393,9 @@ void FMediaProfileEditor::StartCapturingMediaOutput(UMediaOutput* InMediaOutput)
 	}
 	
 	bool bStartedCapture = false;
-	for (const FMediaFrameworkCaptureCameraViewportCameraOutputInfo& ViewportOutputInfo : CaptureSettings->ViewportCaptures)
+	for (int32 Index = 0; Index < CaptureSettings->ViewportCaptures.Num(); ++Index)
 	{
+		const FMediaFrameworkCaptureCameraViewportCameraOutputInfo& ViewportOutputInfo = CaptureSettings->ViewportCaptures[Index];
 		if (ViewportOutputInfo.MediaOutput != InMediaOutput)
 		{
 			continue;
@@ -362,25 +405,21 @@ void FMediaProfileEditor::StartCapturingMediaOutput(UMediaOutput* InMediaOutput)
 		NewLiveCapture.MediaOutput = InMediaOutput;
 		NewLiveCapture.MediaCapture = TStrongObjectPtr(InMediaOutput->CreateMediaCapture());
 
-		if (NewLiveCapture.MediaCapture.IsValid())
+		if (!NewLiveCapture.MediaCapture.IsValid())
 		{
-			// TODO: Hook this up to a slate widget when the media viewport tab is created
-			NewLiveCapture.ViewportClient = FMediaFrameworkCaptureLevelEditorViewportClient::CreateViewportClientForCapture(ViewportOutputInfo);
-			NewLiveCapture.Viewport = MakeShared<FSceneViewport>(NewLiveCapture.ViewportClient.Get(), nullptr);
-			NewLiveCapture.ViewportClient->Viewport = NewLiveCapture.Viewport.Get();
+			continue;
+		}
+		
+		NewLiveCapture.ViewportCapture = FindOrCreateViewportCaptureImpl(Index);
+		if (!NewLiveCapture.ViewportCapture.IsValid())
+		{
+			continue;
+		}
 
-			const FIntPoint TargetSize = MediaProfileEditor::GetTargetSize(InMediaOutput);
-			NewLiveCapture.Viewport->SetFixedViewportSize(TargetSize.X, TargetSize.Y);
-			
-			if (NewLiveCapture.MediaCapture->CaptureSceneViewport(NewLiveCapture.Viewport, ViewportOutputInfo.CaptureOptions))
-			{
-				LiveCaptures.Add(NewLiveCapture);
-				bStartedCapture = true;
-			}
-			else
-			{
-				NewLiveCapture.ViewportClient->Viewport = nullptr;
-			}
+		if (NewLiveCapture.MediaCapture->CaptureSceneViewport(NewLiveCapture.ViewportCapture->Viewport, ViewportOutputInfo.CaptureOptions))
+		{
+			LiveCaptures.Add(NewLiveCapture);
+			bStartedCapture = true;
 		}
 	}
 
@@ -397,7 +436,6 @@ void FMediaProfileEditor::StartCapturingMediaOutput(UMediaOutput* InMediaOutput)
 
 		if (NewLiveCapture.MediaCapture.IsValid())
 		{
-			// TODO: When the media viewport tab is configured, hook up the capture to a widget
 			if (NewLiveCapture.MediaCapture->CaptureTextureRenderTarget2D(RenderTargetOutputInfo.RenderTarget.Get(), RenderTargetOutputInfo.CaptureOptions))
 			{
 				LiveCaptures.Add(NewLiveCapture);
@@ -414,6 +452,11 @@ void FMediaProfileEditor::StartCapturingMediaOutput(UMediaOutput* InMediaOutput)
 
 void FMediaProfileEditor::StopCapturingMediaOutput(UMediaOutput* InMediaOutput)
 {
+	if (!InMediaOutput)
+	{
+		return;
+	}
+	
 	for (FLiveMediaCapture& Capture : LiveCaptures)
 	{
 		if (Capture.MediaOutput.Get() != InMediaOutput)
@@ -430,11 +473,6 @@ void FMediaProfileEditor::StopCapturingMediaOutput(UMediaOutput* InMediaOutput)
 		{
 			Capture.MediaOutput->OnOutputModified().RemoveAll(this);
 		}
-
-		if (Capture.ViewportClient.IsValid())
-		{
-			Capture.ViewportClient->Viewport = nullptr;
-		}
 	}
 
 	LiveCaptures.RemoveAll([InMediaOutput](const FLiveMediaCapture& Capture)
@@ -456,11 +494,6 @@ void FMediaProfileEditor::StopCapturingAllMediaOutputs()
 		{
 			Capture.MediaOutput->OnOutputModified().RemoveAll(this);
 		}
-
-		if (Capture.ViewportClient.IsValid())
-		{
-			Capture.ViewportClient->Viewport = nullptr;
-		}
 	}
 
 	LiveCaptures.Empty();
@@ -501,7 +534,7 @@ TSharedRef<SDockTab> FMediaProfileEditor::SpawnTab_MediaOutput(const FSpawnTabAr
 	return SNew(SDockTab)
 		.Label(LOCTEXT("MediaOutputTabLabel", "Media Output"))
 		[
-			SNullWidget::NullWidget
+			SAssignNew(ViewportPanel, SMediaProfileViewport, SharedThis(this))
 		];
 }
 
@@ -511,6 +544,7 @@ TSharedRef<SDockTab> FMediaProfileEditor::SpawnTab_MediaTree(const FSpawnTabArgs
 		.Label(LOCTEXT("MediaTreeTabLabel", "Media"))
 		[
 			SNew(SMediaProfileSourcesTreeView, MediaProfileBeingEdited)
+			.OnMediaItemDeleted(this, &FMediaProfileEditor::OnMediaItemDeleted)
 			.OnSelectedMediaItemsChanged(this, &FMediaProfileEditor::OnSelectedMediaItemsChanged)
 		];
 }
@@ -533,6 +567,24 @@ TSharedRef<SDockTab> FMediaProfileEditor::SpawnTab_Timecode(const FSpawnTabArgs&
 		];
 }
 
+void FMediaProfileEditor::OnMediaItemDeleted(UClass* InMediaType, int32 InMediaItemIndex)
+{
+	if (InMediaType->IsChildOf<UMediaSource>())
+	{
+		constexpr bool bDestroyMediaPlayer = true;
+		MediaProfileBeingEdited->GetPlaybackManager()->CloseSourceFromIndex(InMediaItemIndex, bDestroyMediaPlayer);
+	}
+	else if (InMediaType->IsChildOf<UMediaOutput>())
+	{
+		StopCapturingMediaOutput(MediaProfileBeingEdited->GetMediaOutput(InMediaItemIndex));
+	}
+
+	if (ViewportPanel.IsValid())
+	{
+		ViewportPanel->ForceClearMediaItem(InMediaType, InMediaItemIndex);
+	}
+}
+
 void FMediaProfileEditor::OnSelectedMediaItemsChanged(const TArray<int32>& SelectedMediaSources, const TArray<int32>& SelectedMediaOutputs)
 {
 	if (!DetailsPanel.IsValid())
@@ -541,6 +593,7 @@ void FMediaProfileEditor::OnSelectedMediaItemsChanged(const TArray<int32>& Selec
 	}
 
 	DetailsPanel->SetSelectedMediaItems(SelectedMediaSources, SelectedMediaOutputs);
+	ViewportPanel->SetSelectedMediaItems(SelectedMediaSources, SelectedMediaOutputs);
 }
 
 void FMediaProfileEditor::OnMediaOutputModified(UMediaOutput* InMediaOutput)
@@ -587,9 +640,15 @@ void FMediaProfileEditor::OnPostPIEStarted(bool bIsSimulating)
 {
 	for (FLiveMediaCapture& Capture : LiveCaptures)
 	{
-		if (Capture.ViewportClient.IsValid())
+		if (!Capture.ViewportCapture.IsValid())
 		{
-			Capture.ViewportClient->SetPIE(true);
+			continue;
+		}
+
+		TSharedPtr<FMediaFrameworkCaptureLevelEditorViewportClient> ViewportClient = StaticCastSharedPtr<FMediaFrameworkCaptureLevelEditorViewportClient>(Capture.ViewportCapture->ViewportClient);
+		if (ViewportClient.IsValid())
+		{
+			ViewportClient->SetPIE(true);
 		}
 	}
 }
@@ -598,9 +657,15 @@ void FMediaProfileEditor::OnPrePIEEnded(bool bIsSimulating)
 {
 	for (FLiveMediaCapture& Capture : LiveCaptures)
 	{
-		if (Capture.ViewportClient.IsValid())
+		if (!Capture.ViewportCapture.IsValid())
 		{
-			Capture.ViewportClient->SetPIE(false);
+			continue;
+		}
+
+		TSharedPtr<FMediaFrameworkCaptureLevelEditorViewportClient> ViewportClient = StaticCastSharedPtr<FMediaFrameworkCaptureLevelEditorViewportClient>(Capture.ViewportCapture->ViewportClient);
+		if (ViewportClient.IsValid())
+		{
+			ViewportClient->SetPIE(false);
 		}
 	}
 }
@@ -678,13 +743,13 @@ void FMediaProfileEditor::OnObjectModified(UObject* Object)
 	const UMediaFrameworkEditorCaptureSettings* GlobalCaptureSettings = GetDefault<UMediaFrameworkEditorCaptureSettings>();
 	if (Object == GlobalCaptureSettings)
 	{
-		bool bCachedSaveCaptureSetingsInWorld = GlobalCaptureSettings->bSaveCaptureSetingsInWorld;
+		bool bCachedSaveCaptureSettingsInWorld = GlobalCaptureSettings->bSaveCaptureSetingsInWorld;
 		
 		// Defer invoking the delegate to ensure any modifications to the global settings complete
-		GEditor->GetTimerManager()->SetTimerForNextTick([this, bCachedSaveCaptureSetingsInWorld, GlobalCaptureSettings]()
+		GEditor->GetTimerManager()->SetTimerForNextTick([this, bCachedSaveCaptureSettingsInWorld, GlobalCaptureSettings]()
 		{
 			// If the capture settings object has changed to or from the global settings, invoke the CaptureSettingsObjectChanged delegate
-			if (bCachedSaveCaptureSetingsInWorld != GlobalCaptureSettings->bSaveCaptureSetingsInWorld)
+			if (bCachedSaveCaptureSettingsInWorld != GlobalCaptureSettings->bSaveCaptureSetingsInWorld)
 			{
 				OnCaptureSettingsObjectChanged.Broadcast();
 			}
@@ -692,4 +757,58 @@ void FMediaProfileEditor::OnObjectModified(UObject* Object)
 	}
 }
 
+TSharedPtr<FMediaProfileEditor::FViewportCapture> FMediaProfileEditor::FindOrCreateViewportCaptureImpl(int32 ViewportOutputInfoIndex)
+{
+	// Take the opportunity to clear any invalid viewport captures from the list
+	TArray<int32> ActiveViewportCaptureKeys;
+	ActiveViewportCaptures.GetKeys(ActiveViewportCaptureKeys);
+	for (int32 ViewportCaptureKey : ActiveViewportCaptureKeys)
+	{
+		if (!ActiveViewportCaptures[ViewportCaptureKey].IsValid())
+		{
+			ActiveViewportCaptures.Remove(ViewportCaptureKey);
+		}
+	}
+	
+	UMediaFrameworkWorldSettingsAssetUserData* CaptureSettings = FindOrAddMediaFrameworkCaptureSettings();
+	if (!CaptureSettings)
+	{
+		return nullptr;
+	}
+
+	if (!CaptureSettings->ViewportCaptures.IsValidIndex(ViewportOutputInfoIndex))
+	{
+		return nullptr;
+	}
+	
+	const FMediaFrameworkCaptureCameraViewportCameraOutputInfo& ViewportOutputInfo = CaptureSettings->ViewportCaptures[ViewportOutputInfoIndex];
+	if (!ActiveViewportCaptures.Contains(ViewportOutputInfoIndex))
+	{
+		if (!ViewportOutputInfo.MediaOutput) 
+		{
+			return nullptr;
+		}
+
+		TSharedPtr<FViewportCapture> NewViewportCapture = MakeShared<FViewportCapture>();
+		NewViewportCapture->MediaOutput = ViewportOutputInfo.MediaOutput;
+		NewViewportCapture->ViewportClient = FMediaFrameworkCaptureLevelEditorViewportClient::CreateViewportClientForCapture(ViewportOutputInfo);
+
+		NewViewportCapture->ViewportWidget = SNew(SViewport)
+			.RenderDirectlyToWindow(false)
+			.IsEnabled(FSlateApplication::Get().GetNormalExecutionAttribute())
+			.EnableGammaCorrection(false)
+			.EnableBlending(false);
+		
+		NewViewportCapture->Viewport = MakeShared<FSceneViewport>(NewViewportCapture->ViewportClient.Get(), NewViewportCapture->ViewportWidget);
+		NewViewportCapture->ViewportClient->Viewport = NewViewportCapture->Viewport.Get();
+		NewViewportCapture->ViewportWidget->SetViewportInterface(NewViewportCapture->Viewport.ToSharedRef());
+		
+		ActiveViewportCaptures.Add(ViewportOutputInfoIndex, NewViewportCapture);
+
+		return NewViewportCapture;
+	}
+	
+	return ActiveViewportCaptures[ViewportOutputInfoIndex].Pin();
+}
+
 #undef LOCTEXT_NAMESPACE
diff --git a/Engine/Plugins/Media/MediaFrameworkUtilities/Source/MediaFrameworkUtilitiesEditor/Private/AssetEditor/MediaProfileEditor.h b/Engine/Plugins/Media/MediaFrameworkUtilities/Source/MediaFrameworkUtilitiesEditor/Private/AssetEditor/MediaProfileEditor.h
--- a/Engine/Plugins/Media/MediaFrameworkUtilities/Source/MediaFrameworkUtilitiesEditor/Private/AssetEditor/MediaProfileEditor.h
+++ b/Engine/Plugins/Media/MediaFrameworkUtilities/Source/MediaFrameworkUtilitiesEditor/Private/AssetEditor/MediaProfileEditor.h
@@ -5,6 +5,7 @@
 #include "CoreMinimal.h"
 #include "Toolkits/AssetEditorToolkit.h"
 
+class SMediaProfileViewport;
 class FMediaFrameworkCaptureLevelEditorViewportClient;
 class UMediaFrameworkWorldSettingsAssetUserData;
 class UMediaOutput;
@@ -23,6 +24,20 @@ private:
 	static const FName MediaTreeTabId;
 	static const FName DetailsTabId;
 	static const FName TimecodeTabId;
+
+public:
+	/** Stores objects necessary for rendering viewport camera captures to a widget */
+	struct FViewportCapture
+	{
+		TWeakObjectPtr<UMediaOutput> MediaOutput;
+		TSharedPtr<FLevelEditorViewportClient> ViewportClient;
+		TSharedPtr<FSceneViewport> Viewport;
+		TSharedPtr<SViewport> ViewportWidget;
+
+		~FViewportCapture();
+	};
+	
+	DECLARE_MULTICAST_DELEGATE_OneParam(FOnCaptureMethodChanged, UMediaOutput*)
 	
 public:
 	static TSharedRef<FMediaProfileEditor> CreateMediaProfileEditor(const EToolkitMode::Type Mode, const TSharedPtr< IToolkitHost >& InitToolkitHost, UMediaProfile* InMediaProfile);
@@ -45,6 +60,16 @@ public:
 	virtual FLinearColor GetWorldCentricTabColorScale() const override;
 	// End of FAssetEditorToolkit
 
+	/** Gets the media profile being edited */
+	UMediaProfile* GetMediaProfile() const { return MediaProfileBeingEdited; }
+	
+	/**
+	 * Attempts to find an existing viewport capture for the specified media output, or creates one if one does not exist
+	 * @param InMediaOutput The media output to find a viewport capture for
+	 * @return The viewport capture, if found or created, or null if the media output does not have a viewport capture configured
+	 */
+	TSharedPtr<FViewportCapture> FindOrCreateViewportCapture(UMediaOutput* InMediaOutput);
+	
 	/** Determines if the specified media output is currently capturing */
 	bool IsMediaOutputCapturing(UMediaOutput* InMediaOutput) const;
 
@@ -65,6 +90,9 @@ public:
 
 	/** Gets the delegate that is raised when the capture settings object has been changed for any reason */
 	FSimpleMulticastDelegate& GetOnCaptureSettingsObjectChanged() { return OnCaptureSettingsObjectChanged; }
+
+	/** Gets the delegate that is raised when the capture method is changed through the details panel for a media output */
+	FOnCaptureMethodChanged& GetOnCaptureMethodChanged() { return OnCaptureMethodChanged; }
 	
 private:
 	TSharedRef<SDockTab> SpawnTab_MediaOutput(const FSpawnTabArgs& InArgs);
@@ -72,6 +100,7 @@ private:
 	TSharedRef<SDockTab> SpawnTab_Details(const FSpawnTabArgs& InArgs);
 	TSharedRef<SDockTab> SpawnTab_Timecode(const FSpawnTabArgs& InArgs);
 	
+	void OnMediaItemDeleted(UClass* InMediaType, int32 InMediaItemIndex);
 	void OnSelectedMediaItemsChanged(const TArray<int32>& SelectedMediaSources, const TArray<int32>& SelectedMediaOutputs);
 	void OnMediaOutputModified(UMediaOutput* InMediaOutput);
 	void OnRestartCapturesNextTick();
@@ -85,19 +114,35 @@ private:
 	void OnAssetsDeleted(const TArray<UClass*>& DeletedAssetClasses);
 	void OnObjectPreEditChange(UObject* Object, const FEditPropertyChain& PropertyChain);
 	void OnObjectModified(UObject* Object);
+
+	/** Attempts to find an existing viewport capture for the specified viewport output info, or creates one if one does not exist */
+	TSharedPtr<FViewportCapture> FindOrCreateViewportCaptureImpl(int32 ViewportOutputInfoIndex);
 	
 private:
 	/** The media profile being edited by this editor */
 	TObjectPtr<UMediaProfile> MediaProfileBeingEdited;
 
 	TSharedPtr<SMediaProfileDetailsPanel> DetailsPanel;
-
+	TSharedPtr<SMediaProfileViewport> ViewportPanel;
+
+	/**
+	 * A map that stores the viewport captures for specific media capture camera viewport outputs.
+	 * There are two ways a viewport capture is consumed:
+	 *     When a media output is being displayed in the media profile editor viewport
+	 *     When a media output is actively being captured.
+	 * To avoid having to create duplicate media capture structs for each consumer, store active captures here
+	 * that can be consumed by either or both the viewport and the live media capture. Further, store viewport
+	 * captures using weak pointers so that when all consumers of a viewport capture release their viewport capture pointers,
+	 * the viewport captures are cleaned up, with the invalid weak pointers cleaned up at the next opportunity
+	 */
+	TMap<int32, TWeakPtr<FViewportCapture>> ActiveViewportCaptures;
+
+	/** Stores objects necessary for performing a live capture of a media output */
 	struct FLiveMediaCapture
 	{
 		TWeakObjectPtr<UMediaOutput> MediaOutput;
 		TStrongObjectPtr<UMediaCapture> MediaCapture;
-		TSharedPtr<FMediaFrameworkCaptureLevelEditorViewportClient> ViewportClient;
-		TSharedPtr<FSceneViewport> Viewport;
+		TSharedPtr<FViewportCapture> ViewportCapture;
 	};
 	
 	/** List of any live media output captures that the media profile editor is managing */
@@ -108,7 +153,10 @@ private:
 
 	/** Raised when the object that stores the capture settings has changed */
 	FSimpleMulticastDelegate OnCaptureSettingsObjectChanged;
+
+	/** Raised when the capture method is changed on a media output through the details panel */
+	FOnCaptureMethodChanged OnCaptureMethodChanged;
 	
-	/** Handle for callback to restart any modified captures on the nextx engine tick */
+	/** Handle for callback to restart any modified captures on the next engine tick */
 	FTimerHandle RestartCapturesTimerHandle;
 };
diff --git a/Engine/Plugins/Media/MediaFrameworkUtilities/Source/MediaFrameworkUtilitiesEditor/Private/AssetEditor/MediaProfileMediaItemDetailCustomization.cpp b/Engine/Plugins/Media/MediaFrameworkUtilities/Source/MediaFrameworkUtilitiesEditor/Private/AssetEditor/MediaProfileMediaItemDetailCustomization.cpp
--- a/Engine/Plugins/Media/MediaFrameworkUtilities/Source/MediaFrameworkUtilitiesEditor/Private/AssetEditor/MediaProfileMediaItemDetailCustomization.cpp
+++ b/Engine/Plugins/Media/MediaFrameworkUtilities/Source/MediaFrameworkUtilitiesEditor/Private/AssetEditor/MediaProfileMediaItemDetailCustomization.cpp
@@ -3,6 +3,7 @@
 #include "MediaProfileMediaItemDetailCustomization.h"
 
 #include "IDetailChildrenBuilder.h"
+#include "IPropertyUtilities.h"
 #include "IStructureDataProvider.h"
 #include "MediaFrameworkWorldSettingsAssetUserData.h"
 #include "Algo/Count.h"
@@ -253,7 +254,10 @@ protected:
 
 	virtual void CustomizeChildren(TSharedRef<IPropertyHandle> PropertyHandle, IDetailChildrenBuilder& ChildBuilder, IPropertyTypeCustomizationUtils& CustomizationUtils) override
 	{
-		ChildBuilder.AddProperty(PropertyHandle->GetChildHandle(GET_MEMBER_NAME_CHECKED(FMediaFrameworkCaptureCameraViewportCameraOutputInfo, LockedActors)).ToSharedRef());
+		TSharedPtr<IPropertyHandle> LockedActorsHandle = PropertyHandle->GetChildHandle(GET_MEMBER_NAME_CHECKED(FMediaFrameworkCaptureCameraViewportCameraOutputInfo, LockedActors));
+		LockedActorsHandle->SetOnPropertyValueChangedWithData(TDelegate<void(const FPropertyChangedEvent&)>::CreateSP(this, &FMediaFrameworkViewportOutputInfoCustomization::PropertyValueChange));
+		LockedActorsHandle->SetOnChildPropertyValueChangedWithData(TDelegate<void(const FPropertyChangedEvent&)>::CreateSP(this, &FMediaFrameworkViewportOutputInfoCustomization::PropertyValueChange));
+		ChildBuilder.AddProperty(LockedActorsHandle.ToSharedRef());
 
 		// Manually add all capture options properties to avoid an additional property group getting added
 		TSharedPtr<IPropertyHandle> CaptureOptionsHandle = PropertyHandle->GetChildHandle(GET_MEMBER_NAME_CHECKED(FMediaFrameworkCaptureCameraViewportCameraOutputInfo, CaptureOptions));
@@ -262,9 +266,18 @@ protected:
 		for (uint32 Index = 0; Index < NumChildren; ++Index)
 		{
 			TSharedPtr<IPropertyHandle> ChildHandle = CaptureOptionsHandle->GetChildHandle(Index);
+			ChildHandle->SetOnPropertyValueChangedWithData(TDelegate<void(const FPropertyChangedEvent&)>::CreateSP(this, &FMediaFrameworkViewportOutputInfoCustomization::PropertyValueChange));			
 			ChildBuilder.AddProperty(ChildHandle.ToSharedRef());
 		}
 	}
+	
+	void PropertyValueChange(const FPropertyChangedEvent& PropertyChangedEvent)
+	{
+		if (UMediaFrameworkWorldSettingsAssetUserData* MediaCaptureSettings = FMediaProfileEditor::FindOrAddMediaFrameworkCaptureSettings())
+		{
+			FCoreUObjectDelegates::OnObjectPropertyChanged.Broadcast(MediaCaptureSettings, const_cast<FPropertyChangedEvent&>(PropertyChangedEvent));
+		}
+	}
 };
 
 class FMediaFrameworkRenderTargetOutputInfoCustomization : public IPropertyTypeCustomization
@@ -274,7 +287,9 @@ protected:
 
 	virtual void CustomizeChildren(TSharedRef<IPropertyHandle> PropertyHandle, IDetailChildrenBuilder& ChildBuilder, IPropertyTypeCustomizationUtils& CustomizationUtils) override
 	{
-		ChildBuilder.AddProperty(PropertyHandle->GetChildHandle(GET_MEMBER_NAME_CHECKED(FMediaFrameworkCaptureRenderTargetCameraOutputInfo, RenderTarget)).ToSharedRef());
+		TSharedPtr<IPropertyHandle> RenderTargetHandle = PropertyHandle->GetChildHandle(GET_MEMBER_NAME_CHECKED(FMediaFrameworkCaptureRenderTargetCameraOutputInfo, RenderTarget));
+		RenderTargetHandle->SetOnPropertyValueChangedWithData(TDelegate<void(const FPropertyChangedEvent&)>::CreateSP(this, &FMediaFrameworkRenderTargetOutputInfoCustomization::PropertyValueChange));
+		ChildBuilder.AddProperty(RenderTargetHandle.ToSharedRef());
 
 		// Manually add all capture options properties to avoid an additional property group getting added
 		TSharedPtr<IPropertyHandle> CaptureOptionsHandle = PropertyHandle->GetChildHandle(GET_MEMBER_NAME_CHECKED(FMediaFrameworkCaptureRenderTargetCameraOutputInfo, CaptureOptions));
@@ -283,9 +298,18 @@ protected:
 		for (uint32 Index = 0; Index < NumChildren; ++Index)
 		{
 			TSharedPtr<IPropertyHandle> ChildHandle = CaptureOptionsHandle->GetChildHandle(Index);
+			ChildHandle->SetOnPropertyValueChangedWithData(TDelegate<void(const FPropertyChangedEvent&)>::CreateSP(this, &FMediaFrameworkRenderTargetOutputInfoCustomization::PropertyValueChange));
 			ChildBuilder.AddProperty(ChildHandle.ToSharedRef());
 		}
 	}
+
+	void PropertyValueChange(const FPropertyChangedEvent& InPropertyChangedEvent)
+	{
+		if (UMediaFrameworkWorldSettingsAssetUserData* MediaCaptureSettings = FMediaProfileEditor::FindOrAddMediaFrameworkCaptureSettings())
+		{
+			FCoreUObjectDelegates::OnObjectPropertyChanged.Broadcast(MediaCaptureSettings, const_cast<FPropertyChangedEvent&>(InPropertyChangedEvent));
+		}
+	}
 };
 
 FMediaProfileMediaOutputDetailCustomization::~FMediaProfileMediaOutputDetailCustomization()
@@ -377,14 +401,13 @@ void FMediaProfileMediaOutputDetailCustomization::CustomizeDetails(IDetailLayout
 			.FillWidth(1.0f)
 			[
 				SNew(SMediaCaptureMethodComboButton, Objects)
-				.OnMethodChanged_Lambda([this]()
-				{
-					CachedDetailBuilder.Pin()->ForceRefreshDetails();
-				})
+				.OnMethodChanged(this, &FMediaProfileMediaOutputDetailCustomization::OnCaptureMethodChanged)
 			]
 		];
 
 	TArray<TSharedPtr<FStructOnScope>> CaptureSettingStructs;
+	bool bHasViewportCaptures = false;
+	bool bHasRenderTargetCaptures = false;
 	if (UMediaFrameworkWorldSettingsAssetUserData* MediaCaptureSettings = FMediaProfileEditor::FindOrAddMediaFrameworkCaptureSettings())
 	{
 		for (const TWeakObjectPtr<UObject>& WeakObject : Objects)
@@ -405,7 +428,11 @@ void FMediaProfileMediaOutputDetailCustomization::CustomizeDetails(IDetailLayout
 			{
 				if (ViewportOutputInfo.MediaOutput == MediaOutput)
 				{
-					CaptureSettingStructs.Add(MakeShared<FStructOnScope>(FMediaFrameworkCaptureCameraViewportCameraOutputInfo::StaticStruct(), reinterpret_cast<uint8*>(&ViewportOutputInfo)));
+					TSharedPtr<FStructOnScope> Struct = MakeShared<FStructOnScope>(FMediaFrameworkCaptureCameraViewportCameraOutputInfo::StaticStruct(), reinterpret_cast<uint8*>(&ViewportOutputInfo));
+					Struct->SetPackage(MediaCaptureSettings->GetPackage());
+					
+					CaptureSettingStructs.Add(Struct);
+					bHasViewportCaptures = true;
 				}
 			}
 
@@ -413,13 +440,18 @@ void FMediaProfileMediaOutputDetailCustomization::CustomizeDetails(IDetailLayout
 			{
 				if (RenderTargetOutputInfo.MediaOutput == MediaOutput)
 				{
-					CaptureSettingStructs.Add(MakeShared<FStructOnScope>(FMediaFrameworkCaptureRenderTargetCameraOutputInfo::StaticStruct(), reinterpret_cast<uint8*>(&RenderTargetOutputInfo)));
+					TSharedPtr<FStructOnScope> Struct = MakeShared<FStructOnScope>(FMediaFrameworkCaptureRenderTargetCameraOutputInfo::StaticStruct(), reinterpret_cast<uint8*>(&RenderTargetOutputInfo));
+					Struct->SetPackage(MediaCaptureSettings->GetPackage());
+					
+					CaptureSettingStructs.Add(Struct);
+					bHasRenderTargetCaptures = true;
 				}
 			}
 		}
 	}
 
-	if (CaptureSettingStructs.Num())
+	// Only show the capture settings if all objects have capture settings of the same type (viewport or render target)
+	if (CaptureSettingStructs.Num() && bHasViewportCaptures != bHasRenderTargetCaptures)
 	{
 		FAddPropertyParams AddPropertyParams;
 		AddPropertyParams.HideRootObjectNode(true);
@@ -661,4 +693,41 @@ void FMediaProfileMediaOutputDetailCustomization::OnCaptureSettingsObjectChanged
 	}
 }
 
+void FMediaProfileMediaOutputDetailCustomization::OnCaptureMethodChanged()
+{
+	TSharedPtr<FMediaProfileEditor> PinnedMediaProfileEditor = MediaProfileEditor.Pin();
+	if (!PinnedMediaProfileEditor.IsValid())
+	{
+		return;
+	}
+
+	TArray<TWeakObjectPtr<UObject>> Objects;
+	CachedDetailBuilder.Pin()->GetObjectsBeingCustomized(Objects);
+
+	for (const TWeakObjectPtr<UObject>& Object : Objects)
+	{
+		TStrongObjectPtr<UObject> PinnedObject = Object.Pin();
+		if (!PinnedObject.IsValid())
+		{
+			continue;
+		}
+
+		if (!PinnedObject->IsA<UMediaOutput>())
+		{
+			continue;
+		}
+
+		if (PinnedObject->IsA<UDummyMediaOutput>())
+		{
+			continue;
+		}
+
+		UMediaOutput* MediaOutput = Cast<UMediaOutput>(PinnedObject.Get());
+		PinnedMediaProfileEditor->StopCapturingMediaOutput(MediaOutput);
+		PinnedMediaProfileEditor->GetOnCaptureMethodChanged().Broadcast(MediaOutput);
+	}
+	
+	CachedDetailBuilder.Pin()->ForceRefreshDetails();
+}
+
 #undef LOCTEXT_NAMESPACE
diff --git a/Engine/Plugins/Media/MediaFrameworkUtilities/Source/MediaFrameworkUtilitiesEditor/Private/AssetEditor/MediaProfileMediaItemDetailCustomization.h b/Engine/Plugins/Media/MediaFrameworkUtilities/Source/MediaFrameworkUtilitiesEditor/Private/AssetEditor/MediaProfileMediaItemDetailCustomization.h
--- a/Engine/Plugins/Media/MediaFrameworkUtilities/Source/MediaFrameworkUtilitiesEditor/Private/AssetEditor/MediaProfileMediaItemDetailCustomization.h
+++ b/Engine/Plugins/Media/MediaFrameworkUtilities/Source/MediaFrameworkUtilitiesEditor/Private/AssetEditor/MediaProfileMediaItemDetailCustomization.h
@@ -223,7 +223,7 @@ protected:
 		for (int32 Index = 0; Index < Objects.Num(); ++Index)
 		{
 			TMediaType* OriginalMediaObject = Cast<TMediaType>(Objects[Index]);
-			TMediaType* NewMediaObject = NewObject<TMediaType>(MediaProfile.Get(), InClass);
+			TMediaType* NewMediaObject = NewObject<TMediaType>(MediaProfile.Get(), InClass, NAME_None, RF_Public | RF_Transactional);
 			
 			SetMediaObject(OriginalMediaObject, NewMediaObject);
 		}
@@ -338,6 +338,8 @@ private:
 	
 	void OnCaptureSettingsObjectChanged();
 
+	void OnCaptureMethodChanged();
+	
 private:
 	/**
 	 * Stored FStructOnScopes for the media capture structs of all objects being displayed by this customization, which are pulled from
diff --git a/Engine/Plugins/Media/MediaFrameworkUtilities/Source/MediaFrameworkUtilitiesEditor/Private/AssetEditor/SMediaProfileMediaItemDisplay.cpp b/Engine/Plugins/Media/MediaFrameworkUtilities/Source/MediaFrameworkUtilitiesEditor/Private/AssetEditor/SMediaProfileMediaItemDisplay.cpp
new file mode 100644
--- /dev/null
+++ b/Engine/Plugins/Media/MediaFrameworkUtilities/Source/MediaFrameworkUtilitiesEditor/Private/AssetEditor/SMediaProfileMediaItemDisplay.cpp
@@ -0,0 +1,751 @@
+﻿// Copyright Epic Games, Inc. All Rights Reserved.
+
+
+#include "SMediaProfileMediaItemDisplay.h"
+
+#include "LevelEditorViewport.h"
+#include "MediaFrameworkWorldSettingsAssetUserData.h"
+#include "MediaPlayer.h"
+#include "MediaTexture.h"
+#include "ShowFlagMenuCommands.h"
+#include "CaptureTab/SMediaFrameworkCapture.h"
+#include "Engine/TextureRenderTarget2D.h"
+#include "Profile/MediaProfilePlaybackManager.h"
+#include "UI/MediaFrameworkUtilitiesEditorStyle.h"
+#include "ViewportToolbar/UnrealEdViewportToolbar.h"
+#include "Widgets/SMediaImage.h"
+#include "Widgets/SViewport.h"
+#include "Widgets/Images/SImage.h"
+#include "Widgets/Text/STextBlock.h"
+
+#define LOCTEXT_NAMESPACE "SMediaProfileMediaItemDisplay"
+
+void SMediaProfileMediaSourceDisplay::Construct(const FArguments& InArgs, const TSharedPtr<SMediaProfileViewport>& InOwningViewport)
+{
+	SMediaProfileMediaItemDisplayBase::FArguments BaseArgs;
+	BaseArgs.MediaProfileEditor(InArgs._MediaProfileEditor);
+	BaseArgs.PanelIndex(InArgs._PanelIndex);
+	BaseArgs.MediaItemIndex(InArgs._MediaItemIndex);
+	SMediaProfileMediaItemDisplayBase::Construct(BaseArgs, InOwningViewport);
+
+	Overlay->AddSlot()
+	.HAlign(HAlign_Left)
+	.VAlign(VAlign_Bottom)
+	.Padding(16.0f)
+	[
+		SNew(SHorizontalBox)
+
+		+SHorizontalBox::Slot()
+		.AutoWidth()
+		.HAlign(HAlign_Left)
+		.VAlign(VAlign_Center)
+		.Padding(0.0f, 0.0f, 16.0f, 0.0f)
+		[
+			SNew(STextBlock)
+			.Text(this, &SMediaProfileMediaSourceDisplay::GetTimeDurationText)
+		]
+
+		+SHorizontalBox::Slot()
+		.AutoWidth()
+		.HAlign(HAlign_Left)
+		.VAlign(VAlign_Center)
+		[
+			SNew(STextBlock)
+			.Text(this, &SMediaProfileMediaSourceDisplay::GetFramerateText)
+		]
+	];
+}
+
+void SMediaProfileMediaSourceDisplay::ConfigureMediaImage()
+{
+	UMediaSource* MediaSource = GetMediaItem();
+	if (!MediaSource)
+	{
+		MediaImageContainer->SetContent(
+			SNew(SBorder)
+			.HAlign(HAlign_Center)
+			.VAlign(VAlign_Center)
+			.BorderImage(FAppStyle::GetBrush("NoBorder"))
+			[
+				SNew(STextBlock)
+				.Text(LOCTEXT("MediaSourceNotConfiguredLabel", "Media Source not configured"))
+			]);
+		return;
+	}
+
+	UMediaProfile* MediaProfile = GetMediaProfile();
+	if (!MediaProfile)
+	{
+		MediaImageContainer->SetContent(SNullWidget::NullWidget);
+		return;
+	}
+	
+	UMediaTexture* MediaTexture = MediaProfile->GetPlaybackManager()->OpenSourceFromIndex(MediaItemIndex);
+
+	if (!MediaTexture)
+	{
+		MediaImageContainer->SetContent(
+			SNew(SBorder)
+			.HAlign(HAlign_Center)
+			.VAlign(VAlign_Center)
+			.BorderImage(FAppStyle::GetBrush("NoBorder"))
+			[
+				SNew(STextBlock)
+				.Text(LOCTEXT("CouldNotPlayMediaSourceLabel", "Could not play Media Source"))
+			]);
+		return;
+	}
+
+	const FVector2D TextureSize = FVector2D(MediaTexture->GetSurfaceWidth(), MediaTexture->GetSurfaceHeight());
+	MediaImageContainer->SetContent(
+		SNew(SBox)
+		.MinAspectRatio(this, &SMediaProfileMediaSourceDisplay::GetSourceAspectRatio)
+		.MaxAspectRatio(this, &SMediaProfileMediaSourceDisplay::GetSourceAspectRatio)
+		[
+			SNew(SMediaImage, MediaTexture)
+			.BrushImageSize(TextureSize)
+		]);
+}
+
+UMediaSource* SMediaProfileMediaSourceDisplay::GetMediaItem() const
+{
+	UMediaProfile* MediaProfile = GetMediaProfile();
+	if (!MediaProfile)
+	{
+		return nullptr;
+	}
+
+	return MediaProfile->GetMediaSource(MediaItemIndex);
+}
+
+FText SMediaProfileMediaSourceDisplay::GetMediaItemLabel() const
+{
+	UMediaProfile* MediaProfile = GetMediaProfile();
+	if (!MediaProfile)
+	{
+		return FText::GetEmpty();
+	}
+
+	return FText::FromString(MediaProfile->GetLabelForMediaSource(MediaItemIndex));
+}
+
+FText SMediaProfileMediaSourceDisplay::GetBaseMediaTypeLabel() const
+{
+	return LOCTEXT("MediaSourceBaseTypeLabel", "Media Source");
+}
+
+UMediaTexture* SMediaProfileMediaSourceDisplay::GetMediaTexture() const
+{
+	UMediaProfile* MediaProfile = GetMediaProfile();
+	if (!MediaProfile)
+	{
+		return nullptr;
+	}
+	
+	return MediaProfile->GetPlaybackManager()->GetSourceMediaTextureFromIndex(MediaItemIndex);
+}
+
+FVector2D SMediaProfileMediaSourceDisplay::GetSourceImageSize() const
+{
+	if (UMediaTexture* MediaTexture = GetMediaTexture())
+	{
+		return FVector2D(MediaTexture->GetSurfaceWidth(), MediaTexture->GetSurfaceHeight());
+	}
+
+	return FVector2D::ZeroVector;
+}
+
+FOptionalSize SMediaProfileMediaSourceDisplay::GetSourceAspectRatio() const
+{
+	const FVector2D ImageSize = GetSourceImageSize();
+	return (ImageSize.X > 0 && ImageSize.Y > 0) ? FOptionalSize(ImageSize.X / ImageSize.Y) : FOptionalSize();
+}
+
+FText SMediaProfileMediaSourceDisplay::GetTimeDurationText() const
+{
+	UMediaTexture* MediaTexture = GetMediaTexture();
+	if (!MediaTexture)
+	{
+		return FText::GetEmpty();
+	}
+
+	UMediaPlayer* MediaPlayer = MediaTexture->GetMediaPlayer();
+	if (!MediaPlayer)
+	{
+		return FText::GetEmpty();
+	}
+
+	FTimespan Time = MediaPlayer->GetTime();
+	FTimespan Duration = MediaPlayer->GetDuration();
+
+	return FText::Format(LOCTEXT("MediaSourceTimeDurationFormat", "{0} / {1}"),
+		FText::FromString(Time.ToString(TEXT("%h:%m:%s.%f"))),
+		Duration == Duration.MaxValue() ? LOCTEXT("InfinitySymbol", "\u221E") : FText::FromString(Duration.ToString(TEXT("%h:%m:%s.%f"))));
+}
+
+FText SMediaProfileMediaSourceDisplay::GetFramerateText() const
+{
+	UMediaTexture* MediaTexture = GetMediaTexture();
+	if (!MediaTexture)
+	{
+		return FText::GetEmpty();
+	}
+
+	UMediaPlayer* MediaPlayer = MediaTexture->GetMediaPlayer();
+	if (!MediaPlayer)
+	{
+		return FText::GetEmpty();
+	}
+	
+	return FText::Format(LOCTEXT("MediaSourceFramerateFormat", "{0} fps"), MediaPlayer->GetVideoTrackFrameRate(INDEX_NONE, INDEX_NONE));
+}
+
+/** Widget to display a viewport media capture output */
+class SMediaCaptureViewportImage : public SCompoundWidget
+{
+public:
+	SLATE_BEGIN_ARGS(SMediaCaptureViewportImage) {}
+		SLATE_ATTRIBUTE(bool, IsCapturing)
+	SLATE_END_ARGS()
+
+	virtual ~SMediaCaptureViewportImage() override
+	{
+		FEditorDelegates::PrePIEEnded.RemoveAll(this);
+		FEditorDelegates::PostPIEStarted.RemoveAll(this);
+	}
+	
+	void Construct(const FArguments& InArgs, const TSharedPtr<FMediaProfileEditor::FViewportCapture> InViewportCapture, const TArray<TLazyObjectPtr<AActor>>& InLockedActors)
+	{
+		IsCapturing = InArgs._IsCapturing;
+		ViewportCapture = InViewportCapture;
+
+		FEditorDelegates::PostPIEStarted.AddSP(this, &SMediaCaptureViewportImage::OnPostPIEStarted);
+		FEditorDelegates::PrePIEEnded.AddSP(this, &SMediaCaptureViewportImage::OnPrePIEEnded);
+		
+		TSharedPtr<SWidget> Content;
+		if (InViewportCapture.IsValid() && InViewportCapture->ViewportWidget.IsValid())
+		{
+			Content = InViewportCapture->ViewportWidget;
+		}
+		else
+		{
+			Content = SNew(SBorder)
+				.HAlign(HAlign_Center)
+				.VAlign(VAlign_Center)
+				.BorderImage(FAppStyle::GetBrush("NoBorder"))
+				[
+					SNew(STextBlock)
+					.Text(LOCTEXT("ViewportNotConfiguredLabel", "Viewport output not configured for Media Output Capture"))
+				];
+		}
+		
+		ChildSlot
+		[
+			SNew(SOverlay)
+
+			+SOverlay::Slot()
+			[
+				Content.ToSharedRef()
+			]
+			
+			+SOverlay::Slot()
+			.HAlign(HAlign_Left)
+			.VAlign(VAlign_Top)
+			.Padding(16)
+			[
+				SNew(SHorizontalBox)
+
+				+SHorizontalBox::Slot()
+				.AutoWidth()
+				.HAlign(HAlign_Center)
+				.VAlign(VAlign_Center)
+				.Padding(0.0f, 0.0f, 8.0f, 0.0f)
+				[
+					SNew(SImage).Image(this, &SMediaCaptureViewportImage::GetLiveCaptureIcon)
+				]
+
+				+SHorizontalBox::Slot()
+				.AutoWidth()
+				.HAlign(HAlign_Center)
+				.VAlign(VAlign_Center)
+				[
+					SNew(STextBlock).Text(this, &SMediaCaptureViewportImage::GetLiveCaptureText)
+				]
+			]
+		];
+
+		Refresh(InLockedActors);
+	}
+
+	/** Updates the display if the capture's locked actors have been changed */
+	void Refresh(const TArray<TLazyObjectPtr<AActor>>& InLockedActors)
+	{
+		LockedActors.Empty();
+		LockedActors.Reserve(InLockedActors.Num());
+		for (const TLazyObjectPtr<AActor>& ActorRef : InLockedActors)
+		{
+			AActor* Actor = ActorRef.Get();
+			if (Actor)
+			{
+				LockedActors.Add(Actor);
+			}
+		}
+
+		if (LockedActors.Num())
+		{
+			SetLockedActor(0);
+		}
+		else
+		{
+			SetLockedActor(INDEX_NONE);
+		}
+	}
+
+	/** Sets the actor that functions as the camera for the viewport capture */
+	void SetLockedActor(int32 InLockedActorIndex)
+	{
+		CurrentLockedActor = InLockedActorIndex;
+
+		if (!ViewportCapture.IsValid() || !ViewportCapture->ViewportClient.IsValid())
+		{
+			return;
+		}
+
+		AActor* ActorToLockTo = nullptr;
+		if (LockedActors.IsValidIndex(CurrentLockedActor))
+		{
+			AActor* Actor = LockedActors[CurrentLockedActor].Get();
+			ActorToLockTo = Actor;
+			
+			// If we are in PIE and the locked actor is not, find its PIE counterpart to lock to
+			if (GEditor->PlayWorld)
+			{
+				const bool bIsAlreadyPIEActor = Actor->GetOutermost()->HasAnyPackageFlags(PKG_PlayInEditor);
+				if (!bIsAlreadyPIEActor)
+				{
+					ActorToLockTo = EditorUtilities::GetSimWorldCounterpartActor(Actor);
+				}
+			}
+		}
+
+		ViewportCapture->ViewportClient->SetActorLock(ActorToLockTo);
+	}
+	
+private:
+	const FSlateBrush* GetLiveCaptureIcon() const
+	{
+		bool bIsCapturing = IsCapturing.Get(false);
+		return bIsCapturing ? FMediaFrameworkUtilitiesEditorStyle::Get().GetBrush("MediaCapture.Capture") : FAppStyle::GetBrush("Icons.Toolbar.Stop");
+	}
+	
+	FText GetLiveCaptureText() const
+	{
+		bool bIsCapturing = IsCapturing.Get(false);
+		return bIsCapturing ? LOCTEXT("CapturingLabel", "Capturing") : LOCTEXT("NotCapturingLabel", "Not Capturing");
+	}
+
+	/** Raised when Play-in-Editor is started */
+	void OnPostPIEStarted(bool bIsSimulating)
+	{
+		SetLockedActor(CurrentLockedActor);
+	}
+
+	/** Raised when Play-in-Editor is ended */
+	void OnPrePIEEnded(bool bIsSimulating)
+	{
+		SetLockedActor(CurrentLockedActor);
+	}
+	
+private:
+	/** The viewport capture that contains the viewport client and viewport widgets to display */
+	TSharedPtr<FMediaProfileEditor::FViewportCapture> ViewportCapture;
+
+	/** Command list that manages the show flags for the viewport capture */
+	TSharedPtr<FUICommandList> ShowFlagsCommandList;
+
+	/** List of actors that can function as cameras for the viewport capture */
+	TArray<TWeakObjectPtr<AActor>> LockedActors;
+
+	/** The current locked actor that functions as the camera. */
+	int32 CurrentLockedActor = INDEX_NONE;
+
+	/** Attribute for querying if the viewport capture is currently actively being captured */
+	TAttribute<bool> IsCapturing;
+};
+
+/** Widget to display a render target media capture output */
+class SMediaCaptureRenderTargetImage : public SCompoundWidget
+{
+public:
+	SLATE_BEGIN_ARGS(SMediaCaptureRenderTargetImage) {}
+		SLATE_ATTRIBUTE(bool, IsCapturing)
+	SLATE_END_ARGS()
+
+	void Construct(const FArguments& InArgs, UTextureRenderTarget2D* InRenderTarget)
+	{
+		IsCapturing = InArgs._IsCapturing;
+		
+		ChildSlot
+		[
+			SNew(SOverlay)
+
+			+SOverlay::Slot()
+			[
+				SAssignNew(Container, SBox)
+			]
+
+			+SOverlay::Slot()
+			.HAlign(HAlign_Left)
+			.VAlign(VAlign_Top)
+			.Padding(16)
+			[
+				SNew(SHorizontalBox)
+
+				+SHorizontalBox::Slot()
+				.AutoWidth()
+				.HAlign(HAlign_Center)
+				.VAlign(VAlign_Center)
+				.Padding(0.0f, 0.0f, 8.0f, 0.0f)
+				[
+					SNew(SImage).Image(this, &SMediaCaptureRenderTargetImage::GetLiveCaptureIcon)
+				]
+
+				+SHorizontalBox::Slot()
+				.AutoWidth()
+				.HAlign(HAlign_Center)
+				.VAlign(VAlign_Center)
+				[
+					SNew(STextBlock).Text(this, &SMediaCaptureRenderTargetImage::GetLiveCaptureText)
+				]
+			]
+		];
+
+		Refresh(InRenderTarget);
+	}
+
+	/** Updates the display if the capture's render target has been changed */
+	void Refresh(UTextureRenderTarget2D* InRenderTarget)
+	{
+		if (RenderTarget.IsValid() && RenderTarget.Get() == InRenderTarget)
+		{
+			return;
+		}
+
+		RenderTarget = InRenderTarget;
+		if (InRenderTarget)
+		{
+			Container->SetContent
+			(
+				SNew(SMediaImage, InRenderTarget)
+			);
+		}
+		else
+		{
+			Container->SetContent
+			(
+				SNew(SBorder)
+				.HAlign(HAlign_Center)
+				.VAlign(VAlign_Center)
+				.BorderImage(FAppStyle::GetBrush("NoBorder"))
+				[
+					SNew(STextBlock)
+					.Text(LOCTEXT("NoRenderTargetLabel", "No Render Target set for Media Output Capture"))
+				]
+			);
+		}
+	}
+
+private:
+	const FSlateBrush* GetLiveCaptureIcon() const
+	{
+		bool bIsCapturing = IsCapturing.Get(false);
+		return bIsCapturing ? FMediaFrameworkUtilitiesEditorStyle::Get().GetBrush("MediaCapture.Capture") : FAppStyle::GetBrush("Icons.Toolbar.Stop");
+	}
+	
+	FText GetLiveCaptureText() const
+	{
+		bool bIsCapturing = IsCapturing.Get(false);
+		return bIsCapturing ? LOCTEXT("CapturingLabel", "Capturing") : LOCTEXT("NotCapturingLabel", "Not Capturing");
+	}
+	
+private:
+	/** Widget that contains the render target output widget */
+	TSharedPtr<SBox> Container;
+
+	/** The render target being used by the capture */
+	TWeakObjectPtr<UTextureRenderTarget2D> RenderTarget;
+	
+	/** Attribute for querying if the viewport capture is currently actively being captured */
+	TAttribute<bool> IsCapturing;
+};
+
+void SMediaProfileMediaOutputDisplay::Construct(const FArguments& InArgs, const TSharedPtr<SMediaProfileViewport>& InOwningViewport)
+{
+	ShowFlagsCommandList = MakeShared<FUICommandList>();
+	
+	SMediaProfileMediaItemDisplayBase::FArguments BaseArgs;
+	BaseArgs.MediaProfileEditor(InArgs._MediaProfileEditor);
+	BaseArgs.PanelIndex(InArgs._PanelIndex);
+	BaseArgs.MediaItemIndex(InArgs._MediaItemIndex);
+	SMediaProfileMediaItemDisplayBase::Construct(BaseArgs, InOwningViewport);
+
+	if (MediaProfileEditor.IsValid())
+	{
+		MediaProfileEditor.Pin()->GetOnCaptureMethodChanged().AddSP(this, &SMediaProfileMediaOutputDisplay::OnCaptureMethodChanged);
+	}
+}
+
+SMediaProfileMediaOutputDisplay::~SMediaProfileMediaOutputDisplay()
+{
+	if (MediaProfileEditor.IsValid())
+	{
+		MediaProfileEditor.Pin()->GetOnCaptureMethodChanged().RemoveAll(this);
+	}
+}
+
+void SMediaProfileMediaOutputDisplay::ConfigureMediaImage()
+{
+	ImageWidget.Reset();
+	ViewportOutputInfoIndex = INDEX_NONE;
+	RenderTargetOutputInfoIndex = INDEX_NONE;
+
+	// Clear out all commands that have been bound to any viewport clients from viewport captures
+	for (const FShowFlagMenuCommands::FShowFlagCommand& FlagCommand : FShowFlagMenuCommands::Get().GetCommands())
+	{
+		ShowFlagsCommandList->UnmapAction(FlagCommand.ShowMenuItem);
+	}
+	bDisplayShowFlags = false;
+	
+	UMediaOutput* MediaOutput = GetMediaItem();
+	if (!MediaOutput)
+	{
+		MediaImageContainer->SetContent(
+			SNew(SBorder)
+			.HAlign(HAlign_Center)
+			.VAlign(VAlign_Center)
+			.BorderImage(FAppStyle::GetBrush("NoBorder"))
+			[
+				SNew(STextBlock).Text(LOCTEXT("MediaOutputNotConfiguredLabel", "Media Output not configured"))
+			]);
+		return;
+	}
+	
+	if (UMediaFrameworkWorldSettingsAssetUserData* CaptureSettings = FMediaProfileEditor::FindOrAddMediaFrameworkCaptureSettings())
+	{
+		ViewportOutputInfoIndex = CaptureSettings->ViewportCaptures.IndexOfByPredicate(
+		[MediaOutput](const FMediaFrameworkCaptureCameraViewportCameraOutputInfo& OutputInfo)
+			{
+				return OutputInfo.MediaOutput == MediaOutput;
+			});
+		
+		if (ViewportOutputInfoIndex != INDEX_NONE)
+		{
+			TSharedPtr<FMediaProfileEditor::FViewportCapture> ViewportCapture = MediaProfileEditor.IsValid() ? MediaProfileEditor.Pin()->FindOrCreateViewportCapture(GetMediaItem()) : nullptr;
+			
+			FIntPoint TargetSize = MediaOutput->GetRequestedSize();
+			if (TargetSize == UMediaOutput::RequestCaptureSourceSize)
+			{
+				TargetSize = GetMutableDefault<UMediaFrameworkMediaCaptureSettings>()->DefaultViewportCaptureSize;
+			}
+
+			const float AspectRatio = (TargetSize.X > 0 && TargetSize.Y > 0) ? (float)TargetSize.X / (float)TargetSize.Y : 1.77777777f;
+			
+			MediaImageContainer->SetContent(
+				SNew(SBox)
+				.MinAspectRatio(AspectRatio)
+				.MaxAspectRatio(AspectRatio)
+				[
+					SAssignNew(ImageWidget, SMediaCaptureViewportImage, ViewportCapture, CaptureSettings->ViewportCaptures[ViewportOutputInfoIndex].LockedActors)
+					.IsCapturing(this, &SMediaProfileMediaOutputDisplay::IsCapturingMediaOutput)
+				]
+			);
+
+			FShowFlagMenuCommands::Get().BindCommands(*ShowFlagsCommandList.Get(), ViewportCapture->ViewportClient);
+			bDisplayShowFlags = true;
+			return;
+		}
+
+		RenderTargetOutputInfoIndex = CaptureSettings->RenderTargetCaptures.IndexOfByPredicate(
+		[MediaOutput](const FMediaFrameworkCaptureRenderTargetCameraOutputInfo& OutputInfo)
+			{
+				return OutputInfo.MediaOutput == MediaOutput;
+			});
+
+		if (RenderTargetOutputInfoIndex != INDEX_NONE)
+		{
+			FIntPoint TargetSize = MediaOutput->GetRequestedSize();
+			if (TargetSize == UMediaOutput::RequestCaptureSourceSize)
+			{
+				TargetSize = GetMutableDefault<UMediaFrameworkMediaCaptureSettings>()->DefaultViewportCaptureSize;
+			}
+
+			const float AspectRatio = (TargetSize.X > 0 && TargetSize.Y > 0) ? (float)TargetSize.X / (float)TargetSize.Y : 1.77777777f;
+			MediaImageContainer->SetContent(
+				SNew(SBox)
+				.MinAspectRatio(AspectRatio)
+				.MaxAspectRatio(AspectRatio)
+				[
+					SAssignNew(ImageWidget, SMediaCaptureRenderTargetImage, CaptureSettings->RenderTargetCaptures[RenderTargetOutputInfoIndex].RenderTarget)
+					.IsCapturing(this, &SMediaProfileMediaOutputDisplay::IsCapturingMediaOutput)
+				]
+			);
+			return;
+		}
+	}
+
+	MediaImageContainer->SetContent(
+		SNew(SBorder)
+		.HAlign(HAlign_Center)
+		.VAlign(VAlign_Center)
+		.BorderImage(FAppStyle::GetBrush("NoBorder"))
+		[
+			SNew(STextBlock).Text(LOCTEXT("MediaCaptureNotConfiguredLabel", "Media capture not configured"))
+		]);
+}
+
+UMediaOutput* SMediaProfileMediaOutputDisplay::GetMediaItem() const
+{
+	UMediaProfile* MediaProfile = GetMediaProfile();
+	if (!MediaProfile)
+	{
+		return nullptr;
+	}
+
+	return MediaProfile->GetMediaOutput(MediaItemIndex); 
+}
+
+FText SMediaProfileMediaOutputDisplay::GetMediaItemLabel() const
+{
+	UMediaProfile* MediaProfile = GetMediaProfile();
+	if (!MediaProfile)
+	{
+		return FText::GetEmpty();
+	}
+
+	return FText::FromString(MediaProfile->GetLabelForMediaOutput(MediaItemIndex));
+}
+
+FText SMediaProfileMediaOutputDisplay::GetBaseMediaTypeLabel() const
+{
+	return LOCTEXT("MediaOutputBaseTypeLabel", "Media Output");
+}
+
+void SMediaProfileMediaOutputDisplay::AddToolbarEntries(FToolMenuSection& Section)
+{
+	FToolMenuEntry& ShowSubmenuEntry = Section.AddEntry(UE::UnrealEd::CreateShowSubmenu(
+		FNewToolMenuDelegate::CreateLambda([](UToolMenu* Submenu)
+		{
+			FShowFlagMenuCommands::Get().BuildShowFlagsMenu(Submenu);
+		})));
+
+	ShowSubmenuEntry.Visibility = TAttribute<bool>::CreateLambda([this] { return bDisplayShowFlags; });
+}
+
+void SMediaProfileMediaOutputDisplay::AppendToolbarCommandList(FToolMenuContext& Context)
+{
+	Context.AppendCommandList(ShowFlagsCommandList);
+}
+
+void SMediaProfileMediaOutputDisplay::OnObjectPropertyChanged(UObject* InObject, FPropertyChangedEvent& InPropertyChangedEvent)
+{
+	SMediaProfileMediaItemDisplayBase::OnObjectPropertyChanged(InObject, InPropertyChangedEvent);
+	
+	if (UMediaFrameworkWorldSettingsAssetUserData* CaptureSettings = Cast<UMediaFrameworkWorldSettingsAssetUserData>(InObject))
+	{
+		if (CaptureSettings != FMediaProfileEditor::FindOrAddMediaFrameworkCaptureSettings())
+		{
+			return;
+		}
+		
+		if (InPropertyChangedEvent.GetPropertyName() == GET_MEMBER_NAME_CHECKED(FMediaFrameworkCaptureRenderTargetCameraOutputInfo, RenderTarget))
+		{
+			if (RenderTargetOutputInfoIndex == INDEX_NONE || !CaptureSettings->RenderTargetCaptures.IsValidIndex(RenderTargetOutputInfoIndex))
+			{
+				return;
+			}
+	
+			TSharedPtr<SMediaCaptureRenderTargetImage> RenderTargetImageWidget = StaticCastSharedPtr<SMediaCaptureRenderTargetImage>(ImageWidget);
+			if (!RenderTargetImageWidget.IsValid())
+			{
+				return;
+			}
+
+			RenderTargetImageWidget->Refresh(CaptureSettings->RenderTargetCaptures[RenderTargetOutputInfoIndex].RenderTarget);
+			return;
+		}
+
+		if (InPropertyChangedEvent.GetPropertyName() == GET_MEMBER_NAME_CHECKED(FMediaFrameworkCaptureCameraViewportCameraOutputInfo, LockedActors))
+		{
+			if (ViewportOutputInfoIndex == INDEX_NONE || !CaptureSettings->ViewportCaptures.IsValidIndex(ViewportOutputInfoIndex))
+			{
+				return;
+			}
+			
+			TSharedPtr<SMediaCaptureViewportImage> ViewportImageWidget = StaticCastSharedPtr<SMediaCaptureViewportImage>(ImageWidget);
+			if (!ViewportImageWidget.IsValid())
+			{
+				return;
+			}
+
+			ViewportImageWidget->Refresh(CaptureSettings->ViewportCaptures[ViewportOutputInfoIndex].LockedActors);
+			return;
+		}
+	}
+}
+
+bool SMediaProfileMediaOutputDisplay::IsCapturingMediaOutput() const
+{
+	TSharedPtr<FMediaProfileEditor> PinnedMediaProfileEditor = MediaProfileEditor.Pin();
+	if (!PinnedMediaProfileEditor.IsValid())
+	{
+		return false;
+	}
+
+	return PinnedMediaProfileEditor->IsMediaOutputCapturing(GetMediaItem());
+}
+
+void SMediaProfileMediaOutputDisplay::OnCaptureMethodChanged(UMediaOutput* MediaOutput)
+{
+	if (MediaOutput != GetMediaItem())
+	{
+		return;
+	}
+
+	ConfigureMediaImage();
+}
+
+void SMediaProfileDummyDisplay::Construct(const FArguments& InArgs, const TSharedPtr<SMediaProfileViewport>& InOwningViewport)
+{
+	SMediaProfileMediaItemDisplayBase::FArguments BaseArgs;
+	BaseArgs.MediaProfileEditor(InArgs._MediaProfileEditor);
+	BaseArgs.PanelIndex(InArgs._PanelIndex);
+	BaseArgs.MediaItemIndex(INDEX_NONE);
+	SMediaProfileMediaItemDisplayBase::Construct(BaseArgs, InOwningViewport);
+}
+
+void SMediaProfileDummyDisplay::ConfigureMediaImage()
+{
+	MediaImageContainer->SetContent(
+		SNew(SBorder)
+		.HAlign(HAlign_Center)
+		.VAlign(VAlign_Center)
+		.BorderImage(FAppStyle::GetBrush("NoBorder"))
+		[
+			SNew(STextBlock)
+			.Text(LOCTEXT("NoMediaItemSelectedLabel", "No media item selected"))
+		]
+	);
+}
+
+FText SMediaProfileDummyDisplay::GetActiveMediaItemLabel() const
+{
+	return LOCTEXT("NoMediaItemSelectedLabel", "No Media Item Selected");
+}
+
+void SMediaProfileDummyDisplay::ChangeActiveMediaItem(UClass* InMediaItemClass, int32 InMediaItemIndex)
+{
+	constexpr bool bRefreshDisplay = true;
+	ChangePanelContents(InMediaItemClass, InMediaItemIndex, bRefreshDisplay);
+}
+
+#undef LOCTEXT_NAMESPACE
diff --git a/Engine/Plugins/Media/MediaFrameworkUtilities/Source/MediaFrameworkUtilitiesEditor/Private/AssetEditor/SMediaProfileMediaItemDisplay.h b/Engine/Plugins/Media/MediaFrameworkUtilities/Source/MediaFrameworkUtilitiesEditor/Private/AssetEditor/SMediaProfileMediaItemDisplay.h
new file mode 100644
--- /dev/null
+++ b/Engine/Plugins/Media/MediaFrameworkUtilities/Source/MediaFrameworkUtilitiesEditor/Private/AssetEditor/SMediaProfileMediaItemDisplay.h
@@ -0,0 +1,679 @@
+﻿// Copyright Epic Games, Inc. All Rights Reserved.
+
+#pragma once
+
+#include "LevelViewportActions.h"
+#include "MediaOutput.h"
+#include "MediaProfileEditor.h"
+#include "MediaSource.h"
+#include "PropertyEditorModule.h"
+#include "SMediaProfileViewport.h"
+#include "ToolMenus.h"
+#include "Framework/MultiBox/MultiBoxBuilder.h"
+#include "Profile/MediaProfile.h"
+#include "Styling/SlateIconFinder.h"
+#include "Widgets/SCompoundWidget.h"
+#include "Widgets/SOverlay.h"
+#include "Widgets/Layout/SBorder.h"
+#include "Widgets/Layout/SSpacer.h"
+
+#include "SMediaProfileMediaItemDisplay.generated.h"
+
+class UMediaTexture;
+class UMediaProfile;
+
+/**
+ * Used as the context object for the viewport toolbar displayed at the top of each display, stores a reference to the viewport widget
+ * that is displaying the toolbar instance
+ */
+UCLASS()
+class UMediaProfileMediaItemDisplayContext : public UObject
+{
+	GENERATED_BODY()
+
+public:
+	template<typename T>
+	TSharedPtr<T> GetDisplayWidget() const
+	{
+		if (DisplayWidget.IsValid())
+		{
+			return StaticCastSharedPtr<T>(DisplayWidget.Pin());
+		}
+
+		return nullptr;
+	}
+	
+	template<typename T>
+	void SetDisplayWidget(const TSharedPtr<T>& InDisplayWidget)
+	{
+		DisplayWidget = InDisplayWidget;
+	}
+	
+private:
+	TWeakPtr<SWidget> DisplayWidget;
+};
+
+#define LOCTEXT_NAMESPACE "SMediaProfileMediaItemDisplay"
+
+/** Widget to display the input or capture from a single media source or output */
+template<typename TMediaItem>
+class SMediaProfileMediaItemDisplayBase : public SCompoundWidget
+{
+public:
+	SLATE_BEGIN_ARGS(SMediaProfileMediaItemDisplayBase) {}
+		SLATE_ARGUMENT(TWeakPtr<FMediaProfileEditor>, MediaProfileEditor)
+		SLATE_ARGUMENT(int32, PanelIndex)
+		SLATE_ARGUMENT(int32, MediaItemIndex)
+	SLATE_END_ARGS()
+		
+	void Construct(const FArguments& InArgs, const TSharedPtr<SMediaProfileViewport>& InOwningViewport)
+	{
+		OwningViewport = InOwningViewport;
+		MediaProfileEditor = InArgs._MediaProfileEditor;
+		MediaItemIndex = InArgs._MediaItemIndex;
+		PanelIndex = InArgs._PanelIndex;
+
+		FCoreUObjectDelegates::OnObjectPropertyChanged.AddSP(this, &SMediaProfileMediaItemDisplayBase::OnObjectPropertyChanged);
+
+		CommandList = MakeShared<FUICommandList>();
+		BindCommands();
+		
+		ChildSlot
+		[
+			SNew(SVerticalBox)
+
+			+SVerticalBox::Slot()
+			.HAlign(HAlign_Fill)
+			.VAlign(VAlign_Center)
+			.AutoHeight()
+			[
+				SNew(SBorder)
+				.BorderImage(FAppStyle::Get().GetBrush("ToolPanel.GroupBorder"))
+				[
+					SNew(SHorizontalBox)
+
+					+SHorizontalBox::Slot()
+					[
+						SNew(SSpacer)
+					]
+
+					+SHorizontalBox::Slot()
+					.AutoWidth()
+					[
+						CreateToolbar()
+					]
+				]
+			]
+
+			+SVerticalBox::Slot()
+			[
+				SAssignNew(Overlay, SOverlay)
+
+				+SOverlay::Slot()
+				[
+					SAssignNew(MediaImageContainer, SBorder)
+					.BorderBackgroundColor(FLinearColor::Black)
+					.BorderImage(FAppStyle::GetBrush("WhiteBrush"))
+				]
+			]
+		];
+
+		ConfigureMediaImage();
+	}
+
+	virtual ~SMediaProfileMediaItemDisplayBase() override
+	{
+		FCoreUObjectDelegates::OnObjectPropertyChanged.RemoveAll(this);
+	}
+
+	virtual FReply OnKeyDown(const FGeometry& MyGeometry, const FKeyEvent& InKeyEvent) override
+	{
+		if (CommandList->ProcessCommandBindings(InKeyEvent))
+		{
+			return FReply::Handled();
+		}
+
+		return FReply::Unhandled();
+	}
+
+	virtual bool SupportsKeyboardFocus() const override { return true; }
+	
+protected:
+	/** Performs any needed configuration to output the media item to the UI */
+	virtual void ConfigureMediaImage() = 0;
+
+	/** Gets the typed media item being displayed */
+	virtual TMediaItem* GetMediaItem() const = 0;
+
+	/** Gets the label to display for the media item */
+	virtual FText GetMediaItemLabel() const = 0;
+
+	/** Gets the label for the base type of the media item (e.g. media source/media output) */
+	virtual FText GetBaseMediaTypeLabel() const = 0;
+
+	UMediaProfile* GetMediaProfile() const
+	{
+		UMediaProfile* MediaProfile = nullptr;
+		if (MediaProfileEditor.IsValid())
+		{
+			MediaProfile = MediaProfileEditor.Pin()->GetMediaProfile();
+		}
+		
+		return MediaProfile;
+	}
+	
+	/** Allows implementations to add their own entries to the viewport toolbar */
+	virtual void AddToolbarEntries(FToolMenuSection& Section) { }
+
+	/** Allows implementations to add their own command lists to the viewport toolbar */
+	virtual void AppendToolbarCommandList(FToolMenuContext& Context) { }
+
+	/** Binds commands to the display's command list */
+	void BindCommands()
+	{
+		CommandList->MapAction(
+			FLevelViewportCommands::Get().ToggleImmersive,
+			FUIAction(
+				FExecuteAction::CreateLambda([this]()
+				{
+					TSharedPtr<SMediaProfileViewport> PinnedOwningViewport = OwningViewport.Pin();
+					if (!PinnedOwningViewport.IsValid())
+					{
+						return;
+					}
+						
+					if (!IsImmersive())
+					{
+						PinnedOwningViewport->SetImmersivePanel(PanelIndex);
+					}
+					else
+					{
+						PinnedOwningViewport->ClearImmersivePanel();
+					}
+				}),
+				FCanExecuteAction(),
+				FIsActionChecked::CreateLambda([this] { return IsImmersive(); })
+			));
+
+		CommandList->MapAction(FLevelViewportCommands::Get().ToggleMaximize,
+			FUIAction(
+				FExecuteAction::CreateLambda([this]()
+				{
+					TSharedPtr<SMediaProfileViewport> PinnedOwningViewport = OwningViewport.Pin();
+					if (!PinnedOwningViewport.IsValid())
+					{
+						return;
+					}
+						
+					if (!IsMaximized())
+					{
+						PinnedOwningViewport->MaximizePanel(PanelIndex);
+					}
+					else
+					{
+						PinnedOwningViewport->RestorePreviousLayout();
+					}
+				}),
+				FCanExecuteAction::CreateLambda([this]()
+				{
+					TSharedPtr<SMediaProfileViewport> PinnedOwningViewport = OwningViewport.Pin();
+					if (!PinnedOwningViewport.IsValid())
+					{
+						return false;
+					}
+
+					return IsMaximized() ? PinnedOwningViewport->CanRestorePreviousLayout() : PinnedOwningViewport->CanMaximizePanel(PanelIndex);
+				})
+			));
+	}
+
+	/** Creates the viewport display's toolbar */
+	TSharedRef<SWidget> CreateToolbar()
+	{
+		const FName ViewportToolbarName = "MediaProfileEditor.ViewportToolbar";
+
+		if (!UToolMenus::Get()->IsMenuRegistered(ViewportToolbarName))
+		{
+			UToolMenu* ToolbarMenu = UToolMenus::Get()->RegisterMenu(ViewportToolbarName, NAME_None, EMultiBoxType::SlimHorizontalToolBar);
+			ToolbarMenu->StyleName = "ViewportToolbar";
+			
+			FToolMenuSection& RightSection = ToolbarMenu->AddSection("Right");
+			RightSection.Alignment = EToolMenuSectionAlign::Last;
+
+			RightSection.AddDynamicEntry(NAME_None, FNewToolMenuSectionDelegate::CreateLambda([](FToolMenuSection& Section)
+			{
+				UMediaProfileMediaItemDisplayContext* Context = Section.FindContext<UMediaProfileMediaItemDisplayContext>();
+				TSharedPtr<SMediaProfileMediaItemDisplayBase<TMediaItem>> ContextWidget = Context->GetDisplayWidget<SMediaProfileMediaItemDisplayBase<TMediaItem>>();
+				
+				if (!ContextWidget.IsValid())
+				{
+					return;
+				}
+				
+				Section.AddEntry(ContextWidget->CreateMediaItemsComboButton());
+				ContextWidget->AddToolbarEntries(Section);
+				Section.AddEntry(ContextWidget->CreateViewportArrangementMenu());
+			}));
+		}
+		
+		FToolMenuContext Context;
+		{
+			Context.AppendCommandList(CommandList);
+			AppendToolbarCommandList(Context);
+			
+			UMediaProfileMediaItemDisplayContext* ContextObject = NewObject<UMediaProfileMediaItemDisplayContext>();
+			ContextObject->SetDisplayWidget<SMediaProfileMediaItemDisplayBase<TMediaItem>>(SharedThis(this));
+			Context.AddObject(ContextObject);
+		}
+
+		return UToolMenus::Get()->GenerateWidget(ViewportToolbarName, Context);
+	}
+
+	/** Creates the combo button used to change which media item is being displayed in the current panel */
+	FToolMenuEntry CreateMediaItemsComboButton()
+	{
+		return FToolMenuEntry::InitSubMenu(
+			TEXT("ActiveMediaItem"),
+			TAttribute<FText>::CreateSP(this, &SMediaProfileMediaItemDisplayBase::GetActiveMediaItemLabel),
+			TAttribute<FText>(),
+			FNewToolMenuDelegate::CreateSP(this, &SMediaProfileMediaItemDisplayBase::GetMediaItemsDropdownContent),
+			false,
+			TAttribute<FSlateIcon>::CreateSP(this, &SMediaProfileMediaItemDisplayBase::GetActiveMediaItemIcon));	
+	}
+
+	/** Creates the dropdown list for the media items combo box, which lists all possible media items that can be displayed in the panel */
+	void GetMediaItemsDropdownContent(UToolMenu* ToolMenu)
+	{
+		UMediaProfile* MediaProfile = GetMediaProfile();
+
+		FToolMenuSection& MediaSourcesSection = ToolMenu->AddSection(TEXT("MediaSourcesSection"), LOCTEXT("MediaSourcesSectionLabel", "Media Sources"));
+		for (int32 Index = 0; Index < MediaProfile->NumMediaSources(); ++Index)
+		{
+			UMediaSource* MediaSource = MediaProfile->GetMediaSource(Index);
+			UClass* Class = MediaSource ? MediaSource->GetClass() : UMediaSource::StaticClass();
+				
+			const FText MediaSourceLabel = MediaSource ? FText::FromString(MediaProfile->GetLabelForMediaSource(Index)) : LOCTEXT("NoMediaSourceLabel", "No Media Source Set");
+
+			MediaSourcesSection.AddMenuEntry(
+				FName(FString::Format(TEXT("MediaSourceEntry_{0}"), { Index })),
+				FText::Format(LOCTEXT("MediaItemMenuEntryFormat", "{0}: {1}"), Index + 1, MediaSourceLabel),
+				TAttribute<FText>(),
+				MediaSource ? FSlateIconFinder::FindIconForClass(Class) : FSlateIcon(),
+				FUIAction(
+					FExecuteAction::CreateSP(this, &SMediaProfileMediaItemDisplayBase<TMediaItem>::ChangeActiveMediaItem, Class, Index),
+					FCanExecuteAction::CreateLambda([bIsMediaSourceSet = MediaSource != nullptr]()
+					{
+						return bIsMediaSourceSet;
+					})));
+		}
+
+		FToolMenuSection& MediaOutputsSection = ToolMenu->AddSection(TEXT("MediaOutputsSection"), LOCTEXT("MediaOutputsSectionLabel", "Media Outputs"));
+		for (int32 Index = 0; Index < MediaProfile->NumMediaSources(); ++Index)
+		{
+			UMediaOutput* MediaOutput = MediaProfile->GetMediaOutput(Index);
+			UClass* Class = MediaOutput ? MediaOutput->GetClass() : UMediaOutput::StaticClass(); 
+
+			const FText MediaOutputLabel = MediaOutput ? FText::FromString(MediaProfile->GetLabelForMediaOutput(Index)) : LOCTEXT("NoMediaOutputLabel", "No Media Output Set");
+			MediaOutputsSection.AddMenuEntry(
+				FName(FString::Format(TEXT("MediaOutputEntry_{0}"), { Index })),
+				FText::Format(LOCTEXT("MediaItemMenuEntryFormat", "{0}: {1}"), Index + 1, MediaOutputLabel),
+				TAttribute<FText>(),
+				MediaOutput ? FSlateIconFinder::FindIconForClass(Class) : FSlateIcon(),
+				FUIAction(
+					FExecuteAction::CreateSP(this, &SMediaProfileMediaItemDisplayBase<TMediaItem>::ChangeActiveMediaItem, Class, Index),
+					FCanExecuteAction::CreateLambda([bIsMediaOutputSet = MediaOutput != nullptr]()
+					{
+						return bIsMediaOutputSet;
+					}))
+				);
+		}
+
+		FToolMenuSection& LockPanelSection = ToolMenu->AddSection(TEXT("LockPanelSection"));
+		LockPanelSection.AddSeparator(NAME_None);
+		
+		FToolMenuEntry& LockButton = LockPanelSection.AddMenuEntry(
+			TEXT("LockDisplayButton"),
+			LOCTEXT("LockDisplayLabel", "Lock Display"),
+			LOCTEXT("LockDisplayTooltip", "Locks this display to the current media item"),
+			TAttribute<FSlateIcon>::CreateSP(this, &SMediaProfileMediaItemDisplayBase::GetLockIcon),
+			FUIAction(
+				FExecuteAction::CreateSP(this, &SMediaProfileMediaItemDisplayBase::OnPanelLockToggled),
+				FCanExecuteAction(),
+				FIsActionChecked::CreateSP(this, &SMediaProfileMediaItemDisplayBase::IsPanelLocked)),
+			EUserInterfaceActionType::ToggleButton);
+			
+		LockButton.SetShowInToolbarTopLevel(true);
+		LockButton.ToolBarData.ResizeParams.AllowClipping = false;
+	}
+
+	/** Creates the viewport arrangement menu, which allows users to change the arrangement of panels in the viewport */
+	FToolMenuEntry CreateViewportArrangementMenu()
+	{
+		FToolMenuEntry Entry = FToolMenuEntry::InitSubMenu(
+			"ViewportArrangementMenu",
+			LOCTEXT("ViewportDropdownLabel", "..."),
+			LOCTEXT("ViewportArrangementTooltip", "Viewport arrangements"),
+			FNewToolMenuDelegate::CreateLambda([this](UToolMenu* Submenu)
+			{
+				TSharedPtr<SMediaProfileViewport> PinnedOwningViewport = OwningViewport.Pin();
+				if (!PinnedOwningViewport)
+				{
+					return;
+				}
+				
+				auto AddLayoutSectionFn = [Submenu, CommandList = PinnedOwningViewport->GetCommandList()](const FName& SectionName, const FText& SectionLabel, const TArray<TSharedPtr<FUICommandInfo>>& Commands)
+				{
+					FToolMenuSection& Section = Submenu->AddSection(SectionName, SectionLabel);
+
+					FSlimHorizontalToolBarBuilder ToolbarBuilder(CommandList, FMultiBoxCustomization::None);
+					ToolbarBuilder.SetLabelVisibility(EVisibility::Collapsed);
+					ToolbarBuilder.SetStyle(&FAppStyle::Get(), "ViewportLayoutToolbar");
+
+					for (const TSharedPtr<FUICommandInfo>& Command : Commands)
+					{
+						ToolbarBuilder.AddToolBarButton(Command);
+					}
+
+					TSharedRef<SWidget> Widget = SNew(SHorizontalBox)
+						+SHorizontalBox::Slot()
+						.AutoWidth()
+						[
+							ToolbarBuilder.MakeWidget()
+						]
+						
+						+SHorizontalBox::Slot()
+						.FillWidth(1)
+						[
+							SNullWidget::NullWidget
+						];
+					
+					const FText Label = FText::GetEmpty();
+					constexpr bool bNoIndent = true;
+					Section.AddEntry(FToolMenuEntry::InitWidget(SectionName, Widget, Label, bNoIndent));
+				};
+
+				FLevelViewportCommands& Commands = FLevelViewportCommands::Get();
+				
+				AddLayoutSectionFn(TEXT("ViewportOnePaneConfigs"), LOCTEXT("OnePaneConfigHeader", "One Pane"), { Commands.ViewportConfig_OnePane });
+				AddLayoutSectionFn(TEXT("ViewportTwoPaneConfigs"), LOCTEXT("TwoPaneConfigHeader", "Two Panes"), { Commands.ViewportConfig_TwoPanesH, Commands.ViewportConfig_TwoPanesV });
+				AddLayoutSectionFn(TEXT("ViewportThreePaneConfigs"), LOCTEXT("ThreePaneConfigHeader", "Three Panes"), {
+					Commands.ViewportConfig_ThreePanesLeft,
+					Commands.ViewportConfig_ThreePanesRight,
+					Commands.ViewportConfig_ThreePanesTop,
+					Commands.ViewportConfig_ThreePanesBottom});
+				AddLayoutSectionFn(TEXT("ViewportFourPaneConfigs"), LOCTEXT("FourPaneConfigHeader", "Four Panes"), {
+					Commands.ViewportConfig_FourPanes2x2,
+					Commands.ViewportConfig_FourPanesLeft,
+					Commands.ViewportConfig_FourPanesRight,
+					Commands.ViewportConfig_FourPanesTop,
+					Commands.ViewportConfig_FourPanesBottom});
+
+				
+				FToolMenuSection& MaximizeSection = Submenu->FindOrAddSection("MaximizeSection");
+			
+				MaximizeSection.AddSeparator("MaximizeSeparator");
+				MaximizeSection.AddEntry(FToolMenuEntry::InitMenuEntry(FLevelViewportCommands::Get().ToggleImmersive));
+				MaximizeSection.AddEntry(CreateMaximizeViewportButton());
+			}));
+
+		Entry.StyleNameOverride = FName("ViewportToolbarViewportSizingSubmenu");
+		Entry.InsertPosition.Position = EToolMenuInsertType::Last;
+		Entry.ToolBarData.LabelOverride = FText();
+		Entry.ToolBarData.ResizeParams.AllowClipping = false;
+
+		return Entry;
+	}
+
+	/** Creates the menu entry to maximize this panel */
+	FToolMenuEntry CreateMaximizeViewportButton()
+	{
+		FToolMenuEntry MaximizeRestoreEntry = FToolMenuEntry::InitMenuEntry(FLevelViewportCommands::Get().ToggleMaximize,
+			TAttribute<FText>::CreateLambda([this]()
+			{
+				return IsMaximized() ?
+					LOCTEXT("MaximizeRestoreLabel_Restore", "Restore All Viewports") :
+					LOCTEXT("MaximizeRestoreLabel_Maximize", "Maximize Viewport");
+			}),
+			TAttribute<FText>::CreateLambda([this]()
+			{
+				return IsMaximized() ?
+					LOCTEXT("MaximizeRestoreTooltip_Restore", "Restores the layout to show all viewports") :
+					LOCTEXT("MaximizeRestoreTooltip_Maximize", "Maximizes this viewport");
+			}),
+			TAttribute<FSlateIcon>::CreateLambda([this]()
+			{
+				return IsMaximized() ?
+					FSlateIcon(FAppStyle::GetAppStyleSetName(), "EditorViewportToolBar.Maximize.Checked") :
+					FSlateIcon(FAppStyle::GetAppStyleSetName(), "EditorViewportToolBar.Maximize.Normal");
+			})
+		);
+				
+		MaximizeRestoreEntry.SetShowInToolbarTopLevel(true);
+		MaximizeRestoreEntry.ToolBarData.ResizeParams.AllowClipping = false;
+		MaximizeRestoreEntry.StyleNameOverride = FName("ViewportToolbarViewportSizingSubmenu");
+
+		return MaximizeRestoreEntry;
+	}
+
+	/** Gets the label to display for this panel's content in the media items combo button */
+	virtual FText GetActiveMediaItemLabel() const
+	{
+		if (GetMediaItem())
+		{
+			return FText::Format(LOCTEXT("ActiveMediaItemLabelFormat", "{0} {1}: {2}"), GetBaseMediaTypeLabel(), MediaItemIndex + 1, GetMediaItemLabel());
+		}
+
+		return FText::Format(LOCTEXT("ActiveMediaItemLabelFormat", "{0} {1}: {2}"), GetBaseMediaTypeLabel(), MediaItemIndex + 1, LOCTEXT("MediaNotConfiguredLabel", "Media not configured"));
+	}
+
+	/** Gets the icon to display for this panel's content in the media items combo button */
+	virtual FSlateIcon GetActiveMediaItemIcon() const
+	{
+		if (TMediaItem* MediaItem = GetMediaItem())
+		{
+			return FSlateIconFinder::FindIconForClass(MediaItem->GetClass());
+		}
+
+		return FSlateIconFinder::FindIconForClass(TMediaItem::StaticClass());
+	}
+
+	/** Raised when the user has selected a new media item in the media items combo button */
+	virtual void ChangeActiveMediaItem(UClass* InMediaItemClass, int32 InMediaItemIndex)
+	{
+		bool bRefreshDisplay = true;
+		if (InMediaItemClass->IsChildOf(TMediaItem::StaticClass()))
+		{
+			// If the media item is the same type as this widget's type, we can update the index and refresh the widget
+			MediaItemIndex = InMediaItemIndex;
+			ConfigureMediaImage();
+			bRefreshDisplay = false;
+		}
+
+		ChangePanelContents(InMediaItemClass, InMediaItemIndex, bRefreshDisplay);
+	}
+
+	/** Gets whether this panel is currently maximized in the viewport */
+	bool IsMaximized() const
+	{
+		TSharedPtr<SMediaProfileViewport> PinnedOwningViewport = OwningViewport.Pin();
+		if (!PinnedOwningViewport.IsValid())
+		{
+			return false;
+		}
+
+		return PinnedOwningViewport->IsPanelMaximized(PanelIndex);
+	}
+
+	/** Gets whether this panel is currently immersive in the viewport */
+	bool IsImmersive() const
+	{
+		TSharedPtr<SMediaProfileViewport> PinnedOwningViewport = OwningViewport.Pin();
+		if (!PinnedOwningViewport.IsValid())
+		{
+			return false;
+		}
+
+		return PinnedOwningViewport->IsPanelImmersive(PanelIndex);
+	}
+
+	/** Changes the panel's contents, recreating the display widgets */
+	void ChangePanelContents(UClass* InMediaItemClass, int32 InMediaItemIndex, bool bRefreshDisplay)
+	{
+		TSharedPtr<SMediaProfileViewport> PinnedOwningViewport = OwningViewport.Pin();
+		if (!PinnedOwningViewport.IsValid())
+		{
+			return;
+		}
+
+		PinnedOwningViewport->SetPanelContents(PanelIndex, InMediaItemClass, InMediaItemIndex, bRefreshDisplay);
+	}
+
+	/** Raised when the lock button is toggled */
+	void OnPanelLockToggled()
+	{
+		TSharedPtr<SMediaProfileViewport> PinnedOwningViewport = OwningViewport.Pin();
+		if (!PinnedOwningViewport.IsValid())
+		{
+			return;
+		}
+
+		PinnedOwningViewport->SetPanelLocked(PanelIndex, !IsPanelLocked());
+	}
+
+	/** Gets whether this panel is locked or not */
+	bool IsPanelLocked() const
+	{
+		return OwningViewport.IsValid() && OwningViewport.Pin()->IsPanelLocked(PanelIndex);
+	}
+
+	/** Gets the lock button icon */
+	FSlateIcon GetLockIcon() const
+	{
+		return IsPanelLocked() ?
+			FSlateIcon(FAppStyle::GetAppStyleSetName(), TEXT("PropertyWindow.Locked")) :
+			FSlateIcon(FAppStyle::GetAppStyleSetName(), TEXT("PropertyWindow.Unlocked"));
+	}
+
+	/** Raised when an object property is changed in the editor, used to potentially refresh the display if the displayed media item is modified */
+	virtual void OnObjectPropertyChanged(UObject* InObject, FPropertyChangedEvent& InPropertyChangedEvent)
+	{
+		if (!InObject)
+		{
+			return;
+		}
+		
+		if (InObject == GetMediaItem())
+		{
+			ConfigureMediaImage();
+		}
+	}
+	
+protected:
+	/** The media profile editor that owns this display */
+	TWeakPtr<FMediaProfileEditor> MediaProfileEditor;
+
+	/** The viewport widget that owns this display */
+	TWeakPtr<SMediaProfileViewport> OwningViewport;
+
+	/** The index of the panel that is displaying this display */
+	int32 PanelIndex = INDEX_NONE;
+
+	/** The index of the media item being displayed by this display */
+	int32 MediaItemIndex = INDEX_NONE;
+
+	TSharedPtr<SOverlay> Overlay;
+	TSharedPtr<SBorder> MediaImageContainer;
+	TSharedPtr<FUICommandList> CommandList;
+};
+
+#undef LOCTEXT_NAMESPACE
+
+/** Implementation of SMediaProfileMediaItemDisplayBase for media sources, which displays the live feed from the selected media source */
+class SMediaProfileMediaSourceDisplay : public SMediaProfileMediaItemDisplayBase<UMediaSource>
+{
+public:
+	SLATE_BEGIN_ARGS(SMediaProfileMediaSourceDisplay) {}
+		SLATE_ARGUMENT(TWeakPtr<FMediaProfileEditor>, MediaProfileEditor)
+		SLATE_ARGUMENT(int32, PanelIndex)
+		SLATE_ARGUMENT(int32, MediaItemIndex)
+	SLATE_END_ARGS()
+
+	void Construct(const FArguments& InArgs, const TSharedPtr<SMediaProfileViewport>& InOwningViewport);
+
+protected:
+	virtual void ConfigureMediaImage() override;
+	
+	virtual UMediaSource* GetMediaItem() const override;
+	virtual FText GetMediaItemLabel() const override;
+	virtual FText GetBaseMediaTypeLabel() const override;
+
+private:
+	UMediaTexture* GetMediaTexture() const;
+	FVector2D GetSourceImageSize() const;
+	FOptionalSize GetSourceAspectRatio() const;
+	
+	FText GetTimeDurationText() const;
+	FText GetFramerateText() const;
+};
+
+/** Implementation of SMediaProfileMediaItemDisplayBase for media outputs, which displays the result from a media capture from the selected media output */
+class SMediaProfileMediaOutputDisplay : public SMediaProfileMediaItemDisplayBase<UMediaOutput>
+{
+public:
+	SLATE_BEGIN_ARGS(SMediaProfileMediaOutputDisplay) {}
+		SLATE_ARGUMENT(TWeakPtr<FMediaProfileEditor>, MediaProfileEditor)
+		SLATE_ARGUMENT(int32, PanelIndex)
+		SLATE_ARGUMENT(int32, MediaItemIndex)
+	SLATE_END_ARGS()
+
+	void Construct(const FArguments& InArgs, const TSharedPtr<SMediaProfileViewport>& InOwningViewport);
+
+	virtual ~SMediaProfileMediaOutputDisplay() override;
+	
+protected:
+	virtual void ConfigureMediaImage() override;
+	
+	virtual UMediaOutput* GetMediaItem() const override;
+	virtual FText GetMediaItemLabel() const override;
+	virtual FText GetBaseMediaTypeLabel() const override;
+
+	virtual void AddToolbarEntries(FToolMenuSection& Section) override;
+	virtual void AppendToolbarCommandList(FToolMenuContext& Context) override;
+	
+	virtual void OnObjectPropertyChanged(UObject* InObject, FPropertyChangedEvent& InPropertyChangedEvent) override;
+
+private:
+	bool IsCapturingMediaOutput() const;
+	
+	void OnCaptureMethodChanged(UMediaOutput* MediaOutput);
+	
+private:
+	TSharedPtr<SWidget> ImageWidget;
+	TSharedPtr<FUICommandList> ShowFlagsCommandList;
+	int32 ViewportOutputInfoIndex = INDEX_NONE;
+	int32 RenderTargetOutputInfoIndex = INDEX_NONE;
+	bool bDisplayShowFlags = false;
+};
+
+/** Dummy display for when a panel needs to be displayed that doesn't contain any media item */
+class SMediaProfileDummyDisplay : public SMediaProfileMediaItemDisplayBase<UObject>
+{
+public:
+	SLATE_BEGIN_ARGS(SMediaProfileDummyDisplay) {}
+		SLATE_ARGUMENT(TWeakPtr<FMediaProfileEditor>, MediaProfileEditor)
+		SLATE_ARGUMENT(int32, PanelIndex)
+	SLATE_END_ARGS()
+		
+	void Construct(const FArguments& InArgs, const TSharedPtr<SMediaProfileViewport>& InOwningViewport);
+
+protected:
+	virtual void ConfigureMediaImage() override;
+	
+	virtual UObject* GetMediaItem() const override { return nullptr; }
+	virtual FText GetMediaItemLabel() const override { return FText::GetEmpty(); }
+	virtual FText GetBaseMediaTypeLabel() const override { return FText::GetEmpty(); }
+
+	virtual FText GetActiveMediaItemLabel() const override;
+	virtual FSlateIcon GetActiveMediaItemIcon() const override { return FSlateIcon(); }
+	virtual void ChangeActiveMediaItem(UClass* InMediaItemClass, int32 InMediaItemIndex) override;
+};
\ No newline at end of file
diff --git a/Engine/Plugins/Media/MediaFrameworkUtilities/Source/MediaFrameworkUtilitiesEditor/Private/AssetEditor/SMediaProfileSourcesTreeView.cpp b/Engine/Plugins/Media/MediaFrameworkUtilities/Source/MediaFrameworkUtilitiesEditor/Private/AssetEditor/SMediaProfileSourcesTreeView.cpp
--- a/Engine/Plugins/Media/MediaFrameworkUtilities/Source/MediaFrameworkUtilitiesEditor/Private/AssetEditor/SMediaProfileSourcesTreeView.cpp
+++ b/Engine/Plugins/Media/MediaFrameworkUtilities/Source/MediaFrameworkUtilitiesEditor/Private/AssetEditor/SMediaProfileSourcesTreeView.cpp
@@ -4,6 +4,7 @@
 
 #include "Editor.h"
 #include "MediaOutput.h"
+#include "MediaProfileEditor.h"
 #include "MediaSource.h"
 #include "ScopedTransaction.h"
 #include "SlateOptMacros.h"
@@ -627,6 +628,7 @@ private:
 void SMediaProfileSourcesTreeView::Construct(const FArguments& InArgs, UMediaProfile* InMediaProfile)
 {
 	MediaProfile = InMediaProfile;
+	OnMediaItemDeleted = InArgs._OnMediaItemDeleted;
 	OnSelectedMediaItemsChanged = InArgs._OnSelectedMediaItemsChanged;
 
 	CommandList = MakeShared<FUICommandList>();
@@ -931,6 +933,8 @@ void SMediaProfileSourcesTreeView::DeleteSelectedMedia()
 		return;
 	}
 
+	MediaTreeView->ClearSelection();
+	
 	FScopedTransaction Transaction(LOCTEXT("DeleteMediaTransaction", "Delete Media"));
 	MediaProfile->Modify();
 
@@ -939,12 +943,33 @@ void SMediaProfileSourcesTreeView::DeleteSelectedMedia()
 	SourcesToRemove.Sort();
 	for (int32 Index = SourcesToRemove.Num() - 1; Index >= 0; --Index)
 	{
+		OnMediaItemDeleted.ExecuteIfBound(UMediaSource::StaticClass(), SourcesToRemove[Index]);
 		MediaProfile->RemoveMediaSource(SourcesToRemove[Index]);
 	}
 
+	UMediaFrameworkWorldSettingsAssetUserData* MediaCaptureSettings = FMediaProfileEditor::FindOrAddMediaFrameworkCaptureSettings();
+
+	auto RemoveMediaCaptures = [MediaCaptureSettings](UMediaOutput* MediaOutput)
+	{
+		if (MediaCaptureSettings && MediaOutput)
+		{
+			MediaCaptureSettings->ViewportCaptures.RemoveAll([MediaOutput](const FMediaFrameworkCaptureCameraViewportCameraOutputInfo& ViewportOutputInfo)
+			{
+				return ViewportOutputInfo.MediaOutput == MediaOutput;
+			});
+
+			MediaCaptureSettings->RenderTargetCaptures.RemoveAll([MediaOutput](const FMediaFrameworkCaptureRenderTargetCameraOutputInfo& RenderTargetOutputInfo)
+			{
+				return RenderTargetOutputInfo.MediaOutput == MediaOutput;
+			});
+		}
+	};
+	
 	OutputsToRemove.Sort();
 	for (int32 Index = OutputsToRemove.Num() - 1; Index >= 0; --Index)
 	{
+		OnMediaItemDeleted.ExecuteIfBound(UMediaOutput::StaticClass(), OutputsToRemove[Index]);
+		RemoveMediaCaptures(MediaProfile->GetMediaOutput(Index));
 		MediaProfile->RemoveMediaOutput(OutputsToRemove[Index]);
 	}
 
@@ -1029,6 +1054,11 @@ void SMediaProfileSourcesTreeView::FillMediaTree()
 	
 	MediaTreeItems.Empty();
 
+	if (MediaTreeView.IsValid())
+	{
+		MediaTreeView->ClearSelection();
+	}
+	
 	if (MediaProfile.IsValid())
 	{
 		FMediaTreeItemPtr MediaSourcesTreeItem = MakeShared<FMediaTreeItem>();
@@ -1238,8 +1268,24 @@ void SMediaProfileSourcesTreeView::OnObjectPropertyChanged(UObject* InObject, FP
 	{
 		// If the object isn't the media profile itself, check to see if it is one of the profile's sources or outputs,
 		// as we want to update the tree view if the sources or outputs change as well
-		if (MediaProfile->FindMediaSourceIndex(Cast<UMediaSource>(InObject)) == INDEX_NONE &&
-			MediaProfile->FindMediaOutputIndex(Cast<UMediaOutput>(InObject)) == INDEX_NONE)
+		bool bFoundMediaItem = false;
+		if (UMediaSource* MediaSource = Cast<UMediaSource>(InObject))
+		{
+			if (MediaProfile->FindMediaSourceIndex(MediaSource) != INDEX_NONE)
+			{
+				bFoundMediaItem = true;
+			}
+		}
+
+		if (UMediaOutput* MediaOutput = Cast<UMediaOutput>(InObject))
+		{
+			if (MediaProfile->FindMediaOutputIndex(MediaOutput) != INDEX_NONE)
+			{
+				bFoundMediaItem = true;
+			}
+		}
+		
+		if (!bFoundMediaItem)
 		{
 			return;
 		}
diff --git a/Engine/Plugins/Media/MediaFrameworkUtilities/Source/MediaFrameworkUtilitiesEditor/Private/AssetEditor/SMediaProfileSourcesTreeView.h b/Engine/Plugins/Media/MediaFrameworkUtilities/Source/MediaFrameworkUtilitiesEditor/Private/AssetEditor/SMediaProfileSourcesTreeView.h
--- a/Engine/Plugins/Media/MediaFrameworkUtilities/Source/MediaFrameworkUtilitiesEditor/Private/AssetEditor/SMediaProfileSourcesTreeView.h
+++ b/Engine/Plugins/Media/MediaFrameworkUtilities/Source/MediaFrameworkUtilitiesEditor/Private/AssetEditor/SMediaProfileSourcesTreeView.h
@@ -18,6 +18,7 @@ class UMediaProfile;
 class SMediaProfileSourcesTreeView : public SCompoundWidget, public FEditorUndoClient
 {
 public:
+	DECLARE_DELEGATE_TwoParams(FOnMediaItemDeleted, UClass* /* MediaItemType */, int32 /* MediaItemIndex */)
 	DECLARE_DELEGATE_TwoParams(FOnSelectedMediaItemsChanged, const TArray<int32>& /* SelectedMediaSources */, const TArray<int32>& /* SelectedMediaOutputs */);
 	
 private:
@@ -47,6 +48,7 @@ private:
 	
 public:
 	SLATE_BEGIN_ARGS(SMediaProfileSourcesTreeView) {}
+		SLATE_EVENT(FOnMediaItemDeleted, OnMediaItemDeleted)
 		SLATE_EVENT(FOnSelectedMediaItemsChanged, OnSelectedMediaItemsChanged)
 	SLATE_END_ARGS()
 	
@@ -114,6 +116,7 @@ private:
 
 	TSharedPtr<SMediaProfileSourcesToolbar> Toolbar;
 
+	FOnMediaItemDeleted OnMediaItemDeleted;
 	FOnSelectedMediaItemsChanged OnSelectedMediaItemsChanged;
 	
 	friend class SMediaTreeItemRow;
diff --git a/Engine/Plugins/Media/MediaFrameworkUtilities/Source/MediaFrameworkUtilitiesEditor/Private/AssetEditor/SMediaProfileViewport.cpp b/Engine/Plugins/Media/MediaFrameworkUtilities/Source/MediaFrameworkUtilitiesEditor/Private/AssetEditor/SMediaProfileViewport.cpp
new file mode 100644
--- /dev/null
+++ b/Engine/Plugins/Media/MediaFrameworkUtilities/Source/MediaFrameworkUtilitiesEditor/Private/AssetEditor/SMediaProfileViewport.cpp
@@ -0,0 +1,608 @@
+﻿// Copyright Epic Games, Inc. All Rights Reserved.
+
+#include "SMediaProfileViewport.h"
+
+#include "Editor.h"
+#include "LevelViewportActions.h"
+#include "SMediaProfileMediaItemDisplay.h"
+#include "ToolMenus.h"
+#include "Framework/Application/SlateApplication.h"
+#include "Profile/MediaProfile.h"
+#include "Widgets/SOverlay.h"
+#include "Widgets/Layout/SBorder.h"
+#include "Widgets/Layout/SSplitter.h"
+#include "Widgets/Text/STextBlock.h"
+
+#define LOCTEXT_NAMESPACE "SMediaProfileViewport"
+
+void SMediaProfileViewport::Construct(const FArguments& InArgs, const TSharedPtr<FMediaProfileEditor> InOwningEditor)
+{
+	MediaProfileEditor = InOwningEditor;
+
+	BindCommands();
+	
+	ChildSlot
+	[
+		SAssignNew(LayoutContainer, SOverlay)
+	];
+}
+
+void SMediaProfileViewport::SetSelectedMediaItems(const TArray<int32>& InSelectedMediaSourceIndices, const TArray<int32>& InSelectedMediaOutputIndices)
+{
+	UMediaProfile* MediaProfile = nullptr;
+	if (MediaProfileEditor.IsValid())
+	{
+		MediaProfile = MediaProfileEditor.Pin()->GetMediaProfile();
+	}
+	
+	if (!MediaProfile)
+	{
+		return;
+	}
+	
+	constexpr int32 MaxNumPanels = 4;
+	const int NumMediaItems = InSelectedMediaSourceIndices.Num() + InSelectedMediaOutputIndices.Num();
+	int32 PanelIndex = 0;
+	
+	for (int32 Index = 0; Index < NumMediaItems; ++Index)
+	{
+		while (PanelIndex < PanelContents.Num() && PanelContents[PanelIndex].bIsLocked)
+		{
+			++PanelIndex;
+		}
+		
+		if (PanelIndex >= MaxNumPanels)
+		{
+			break;
+		}
+
+		const int32 MediaItemIndex = Index < InSelectedMediaSourceIndices.Num() ? InSelectedMediaSourceIndices[Index] : InSelectedMediaOutputIndices[Index - InSelectedMediaSourceIndices.Num()];
+		
+		UClass* MediaType;
+		if (Index < InSelectedMediaSourceIndices.Num())
+		{
+			if (UMediaSource* MediaSource = MediaProfile->GetMediaSource(MediaItemIndex))
+			{
+				MediaType = MediaSource->GetClass();
+			}
+			else
+			{
+				MediaType = UMediaSource::StaticClass();
+			}
+		}
+		else
+		{
+			if (UMediaOutput* MediaOutput = MediaProfile->GetMediaOutput(MediaItemIndex))
+			{
+				MediaType = MediaOutput->GetClass();
+			}
+			else
+			{
+				MediaType = UMediaOutput::StaticClass();
+			}
+		}
+		
+		if (PanelIndex >= PanelContents.Num())
+		{
+			FPanelContent NewPanel;
+			NewPanel.MediaItemIndex = MediaItemIndex;
+			NewPanel.MediaType = MediaType;
+			NewPanel.Widget = nullptr;
+			PanelContents.Add(NewPanel);
+		}
+		else
+		{
+			PanelContents[PanelIndex].MediaItemIndex = MediaItemIndex;
+			PanelContents[PanelIndex].MediaType = MediaType;
+			PanelContents[PanelIndex].Widget = nullptr;
+		}
+
+		++PanelIndex;
+	}
+
+	// Clear out or remove any panel content beyond the panels that are set or locked
+	bool bFoundLockedPanel = false;
+	for (int32 Index = FMath::Min(PanelContents.Num(), MaxNumPanels) - 1; Index >= PanelIndex; --Index)
+	{
+		if (PanelContents[Index].bIsLocked)
+		{
+			bFoundLockedPanel = true;
+			continue;
+		}
+
+		// If there is a locked panel after this panel, simply clear the panel's contents.
+		// Otherwise, remove the panel
+		if (bFoundLockedPanel)
+		{
+			PanelContents[Index].MediaItemIndex = INDEX_NONE;
+			PanelContents[Index].MediaType = nullptr;
+			PanelContents[PanelIndex].Widget = nullptr;
+		}
+		else
+		{
+			PanelContents.RemoveAt(Index);
+		}
+	}
+	
+	CurrentLayout = PanelContents.Num() >= 4 ? EViewportLayout::QuadPanel : (EViewportLayout)(PanelContents.Num() - 1);
+	CurrentOrientation = EViewportOrientation::Left;
+	MaximizedPanelIndex = INDEX_NONE;
+	
+	UpdateDisplays();
+}
+
+void SMediaProfileViewport::ForceClearMediaItem(UClass* InMediaType, int32 InMediaItemIndex)
+{
+	for (FPanelContent& PanelContent : PanelContents)
+	{
+		if (PanelContent.MediaType && PanelContent.MediaType->IsChildOf(InMediaType) && PanelContent.MediaItemIndex == InMediaItemIndex)
+		{
+			PanelContent.MediaType = nullptr;
+			PanelContent.MediaItemIndex = INDEX_NONE;
+			PanelContent.bIsLocked = false;
+		}
+	}
+	
+	UpdateDisplays();
+}
+
+void SMediaProfileViewport::SetPanelContents(int32 InPanelIndex, UClass* InMediaItemClass, int32 InMediaItemIndex, bool bRefreshDisplay)
+{
+	if (!PanelContents.IsValidIndex(InPanelIndex))
+	{
+		return;
+	}
+	
+	PanelContents[InPanelIndex].MediaItemIndex = InMediaItemIndex;
+	PanelContents[InPanelIndex].MediaType = InMediaItemClass;
+
+	if (bRefreshDisplay)
+	{
+		PanelContents[InPanelIndex].Widget = nullptr;
+		UpdateDisplays();
+	}
+}
+
+void SMediaProfileViewport::SetPanelLocked(int32 InPanelIndex, bool bInIsLocked)
+{
+	if (!PanelContents.IsValidIndex(InPanelIndex))
+	{
+		return;
+	}
+
+	PanelContents[InPanelIndex].bIsLocked = bInIsLocked;
+}
+
+bool SMediaProfileViewport::IsPanelLocked(int32 InPanelIndex) const
+{
+	if (!PanelContents.IsValidIndex(InPanelIndex))
+	{
+		return false;
+	}
+
+	return PanelContents[InPanelIndex].bIsLocked;
+}
+
+bool SMediaProfileViewport::CanMaximizePanel(int32 InPanelIndex) const
+{
+	return MaximizedPanelIndex == INDEX_NONE &&
+		ImmersivePanelIndex == INDEX_NONE &&
+		PanelContents.Num() > 1 &&
+		PanelContents.IsValidIndex(InPanelIndex);
+}
+
+void SMediaProfileViewport::MaximizePanel(int32 InPanelIndex)
+{
+	MaximizedPanelIndex = InPanelIndex;
+
+	SavedLayout = CurrentLayout;
+	SavedOrientation = CurrentOrientation;
+	CurrentLayout = EViewportLayout::OnePanel;
+	CurrentOrientation = EViewportOrientation::Left;
+	
+	UpdateDisplays();
+
+	// Set keyboard focus on the panel widget after it has been maximized
+	// so that keyboard shortcuts still work for that panel without the user having to click on it again
+	if (PanelContents.IsValidIndex(MaximizedPanelIndex))
+	{
+		FSlateApplication::Get().ClearKeyboardFocus();
+		FSlateApplication::Get().SetKeyboardFocus(PanelContents[MaximizedPanelIndex].Widget, EFocusCause::SetDirectly);
+	}
+}
+
+bool SMediaProfileViewport::IsPanelMaximized(int32 InPanelIndex) const
+{
+	return MaximizedPanelIndex != INDEX_NONE && MaximizedPanelIndex == InPanelIndex;
+}
+
+bool SMediaProfileViewport::CanRestorePreviousLayout() const
+{
+	return MaximizedPanelIndex != INDEX_NONE &&
+		PanelContents.Num() > 1;
+}
+
+void SMediaProfileViewport::RestorePreviousLayout()
+{
+	if (ImmersivePanelIndex != INDEX_NONE)
+	{
+		ClearImmersivePanel();
+	}
+
+	const int32 PreviousMaximizedPanelIndex = MaximizedPanelIndex;
+	MaximizedPanelIndex = INDEX_NONE;
+	CurrentLayout = SavedLayout;
+	CurrentOrientation = SavedOrientation;
+
+	UpdateDisplays();
+
+	// Reset keyboard focus back the panel widget after it has been repositioned in the panel layout
+	// so that keyboard shortcuts still work for that panel without the user having to click on it again
+	if (PanelContents.IsValidIndex(PreviousMaximizedPanelIndex))
+	{
+		FSlateApplication::Get().ClearKeyboardFocus();
+		FSlateApplication::Get().SetKeyboardFocus(PanelContents[PreviousMaximizedPanelIndex].Widget);
+	}
+}
+
+void SMediaProfileViewport::SetImmersivePanel(int32 InPanelIndex)
+{
+	TSharedPtr<SWindow> ParentWindow = FSlateApplication::Get().FindWidgetWindow(AsShared());
+	if (!ParentWindow.IsValid())
+	{
+		return;
+	}
+
+	if (!PanelContents.IsValidIndex(InPanelIndex))
+	{
+		return;
+	}
+	
+	ImmersivePanelIndex = InPanelIndex;
+	ParentWindow->SetFullWindowOverlayContent(PanelContents[InPanelIndex].Widget);
+	
+	// Set keyboard focus on the panel widget after it has been expanded
+	// so that keyboard shortcuts still work for that panel without the user having to click on it again
+	FSlateApplication::Get().ClearKeyboardFocus();
+	FSlateApplication::Get().SetKeyboardFocus(PanelContents[InPanelIndex].Widget);
+}
+
+bool SMediaProfileViewport::IsPanelImmersive(int32 InPanelIndex) const
+{
+	return ImmersivePanelIndex != INDEX_NONE && ImmersivePanelIndex == InPanelIndex;
+}
+
+void SMediaProfileViewport::ClearImmersivePanel()
+{
+	TSharedPtr<SWindow> ParentWindow = FSlateApplication::Get().FindWidgetWindow(AsShared());
+	if (!ParentWindow.IsValid())
+	{
+		return;
+	}
+
+	const int32 PreviousImmersivePanel = ImmersivePanelIndex;
+	
+	ImmersivePanelIndex = INDEX_NONE;
+	ParentWindow->SetFullWindowOverlayContent(nullptr);
+
+	if (PanelContents.IsValidIndex(PreviousImmersivePanel))
+	{
+		// Reset keyboard focus back the panel widget after it has been repositioned in the panel layout
+		// so that keyboard shortcuts still work for that panel without the user having to click on it again
+		FSlateApplication::Get().ClearKeyboardFocus();
+		FSlateApplication::Get().SetKeyboardFocus(PanelContents[PreviousImmersivePanel].Widget);
+	}
+}
+
+FReply SMediaProfileViewport::OnKeyDown(const FGeometry& MyGeometry, const FKeyEvent& InKeyEvent)
+{
+	FReply Reply = FReply::Unhandled();
+	if (CommandList->ProcessCommandBindings(InKeyEvent))
+	{
+		Reply = FReply::Handled();
+	}
+
+	return Reply;
+}
+
+void SMediaProfileViewport::BindCommands()
+{
+	CommandList = MakeShared<FUICommandList>();
+
+	const FLevelViewportCommands& Commands = FLevelViewportCommands::Get();
+
+	CommandList->MapAction(
+		Commands.ViewportConfig_OnePane,
+		FExecuteAction::CreateSP(this, &SMediaProfileViewport::SetViewportLayout, EViewportLayout::OnePanel, EViewportOrientation::Left),
+		FCanExecuteAction(),
+		FIsActionChecked::CreateSP(this, &SMediaProfileViewport::IsViewportLayoutSet, EViewportLayout::OnePanel, EViewportOrientation::Left));
+
+	CommandList->MapAction(
+		Commands.ViewportConfig_TwoPanesH,
+		FExecuteAction::CreateSP(this, &SMediaProfileViewport::SetViewportLayout, EViewportLayout::TwoPanel, EViewportOrientation::Left),
+		FCanExecuteAction(),
+		FIsActionChecked::CreateSP(this, &SMediaProfileViewport::IsViewportLayoutSet, EViewportLayout::TwoPanel, EViewportOrientation::Left));
+		
+	CommandList->MapAction(
+		Commands.ViewportConfig_TwoPanesV,
+		FExecuteAction::CreateSP(this, &SMediaProfileViewport::SetViewportLayout, EViewportLayout::TwoPanel, EViewportOrientation::Top),
+		FCanExecuteAction(),
+		FIsActionChecked::CreateSP(this, &SMediaProfileViewport::IsViewportLayoutSet, EViewportLayout::TwoPanel, EViewportOrientation::Top));
+
+	CommandList->MapAction(
+		Commands.ViewportConfig_ThreePanesLeft,
+		FExecuteAction::CreateSP(this, &SMediaProfileViewport::SetViewportLayout, EViewportLayout::ThreePanel, EViewportOrientation::Left),
+		FCanExecuteAction(),
+		FIsActionChecked::CreateSP(this, &SMediaProfileViewport::IsViewportLayoutSet, EViewportLayout::ThreePanel, EViewportOrientation::Left));
+
+	CommandList->MapAction(
+		Commands.ViewportConfig_ThreePanesTop,
+		FExecuteAction::CreateSP(this, &SMediaProfileViewport::SetViewportLayout, EViewportLayout::ThreePanel, EViewportOrientation::Top),
+		FCanExecuteAction(),
+		FIsActionChecked::CreateSP(this, &SMediaProfileViewport::IsViewportLayoutSet, EViewportLayout::ThreePanel, EViewportOrientation::Top));
+
+	CommandList->MapAction(
+		Commands.ViewportConfig_ThreePanesRight,
+		FExecuteAction::CreateSP(this, &SMediaProfileViewport::SetViewportLayout, EViewportLayout::ThreePanel, EViewportOrientation::Right),
+		FCanExecuteAction(),
+		FIsActionChecked::CreateSP(this, &SMediaProfileViewport::IsViewportLayoutSet, EViewportLayout::ThreePanel, EViewportOrientation::Right));
+
+	CommandList->MapAction(
+		Commands.ViewportConfig_ThreePanesBottom,
+		FExecuteAction::CreateSP(this, &SMediaProfileViewport::SetViewportLayout, EViewportLayout::ThreePanel, EViewportOrientation::Bottom),
+		FCanExecuteAction(),
+		FIsActionChecked::CreateSP(this, &SMediaProfileViewport::IsViewportLayoutSet, EViewportLayout::ThreePanel, EViewportOrientation::Bottom));
+
+	CommandList->MapAction(
+		Commands.ViewportConfig_FourPanesLeft,
+		FExecuteAction::CreateSP(this, &SMediaProfileViewport::SetViewportLayout, EViewportLayout::OneAndThreePanel, EViewportOrientation::Left),
+		FCanExecuteAction(),
+		FIsActionChecked::CreateSP(this, &SMediaProfileViewport::IsViewportLayoutSet, EViewportLayout::OneAndThreePanel, EViewportOrientation::Left));
+
+	CommandList->MapAction(
+		Commands.ViewportConfig_FourPanesTop,
+		FExecuteAction::CreateSP(this, &SMediaProfileViewport::SetViewportLayout, EViewportLayout::OneAndThreePanel, EViewportOrientation::Top),
+		FCanExecuteAction(),
+		FIsActionChecked::CreateSP(this, &SMediaProfileViewport::IsViewportLayoutSet, EViewportLayout::OneAndThreePanel, EViewportOrientation::Top));
+
+	CommandList->MapAction(
+		Commands.ViewportConfig_FourPanesRight,
+		FExecuteAction::CreateSP(this, &SMediaProfileViewport::SetViewportLayout, EViewportLayout::OneAndThreePanel, EViewportOrientation::Right),
+		FCanExecuteAction(),
+		FIsActionChecked::CreateSP(this, &SMediaProfileViewport::IsViewportLayoutSet, EViewportLayout::OneAndThreePanel, EViewportOrientation::Right));
+
+	CommandList->MapAction(
+		Commands.ViewportConfig_FourPanesBottom,
+		FExecuteAction::CreateSP(this, &SMediaProfileViewport::SetViewportLayout, EViewportLayout::OneAndThreePanel, EViewportOrientation::Bottom),
+		FCanExecuteAction(),
+		FIsActionChecked::CreateSP(this, &SMediaProfileViewport::IsViewportLayoutSet, EViewportLayout::OneAndThreePanel, EViewportOrientation::Bottom));
+
+	CommandList->MapAction(
+		Commands.ViewportConfig_FourPanes2x2,
+		FExecuteAction::CreateSP(this, &SMediaProfileViewport::SetViewportLayout, EViewportLayout::QuadPanel, EViewportOrientation::Left),
+		FCanExecuteAction(),
+		FIsActionChecked::CreateSP(this, &SMediaProfileViewport::IsViewportLayoutSet, EViewportLayout::QuadPanel, EViewportOrientation::Left));
+}
+
+void SMediaProfileViewport::UpdateDisplays()
+{
+	LayoutContainer->ClearChildren();
+	
+	UMediaProfile* MediaProfile = nullptr;
+	if (MediaProfileEditor.IsValid())
+	{
+		MediaProfile = MediaProfileEditor.Pin()->GetMediaProfile();
+	}
+	
+	if (!MediaProfile)
+	{
+		return;
+	}
+	
+	if (PanelContents.Num() <= 0)
+	{
+		LayoutContainer->AddSlot()
+		[
+			SNew(SBorder)
+			.HAlign(HAlign_Center)
+			.VAlign(VAlign_Center)
+			.BorderBackgroundColor(FLinearColor::Black)
+			.BorderImage(FAppStyle::GetBrush("WhiteBrush"))
+			[
+				SNew(STextBlock)
+				.Text(LOCTEXT("NoMediaItemSelectedLabel", "No media item selected"))
+			]
+		];
+		
+		return;
+	}
+	
+	LayoutContainer->AddSlot()
+	[
+		ConstructViewportLayout()
+	];
+}
+
+TSharedRef<SWidget> SMediaProfileViewport::ConstructPanelWidget(int32 InPanelIndex)
+{
+	if (InPanelIndex >= PanelContents.Num())
+	{
+		// If a panel is being requested that doesn't have any stored contents, fill the panel contents array
+		// with empty panels up to that index
+		PanelContents.AddDefaulted(InPanelIndex + 1 - PanelContents.Num());
+	}
+
+	FPanelContent& PanelContent = PanelContents[InPanelIndex];
+	if (PanelContent.Widget.IsValid())
+	{
+		return PanelContent.Widget.ToSharedRef();
+	}
+	
+	if (PanelContent.MediaType && PanelContent.MediaType->IsChildOf<UMediaSource>())
+	{
+		return SAssignNew(PanelContent.Widget, SMediaProfileMediaSourceDisplay, SharedThis(this))
+			.MediaProfileEditor(MediaProfileEditor)
+			.PanelIndex(InPanelIndex)
+			.MediaItemIndex(PanelContent.MediaItemIndex);
+	}
+
+	if (PanelContent.MediaType && PanelContent.MediaType->IsChildOf<UMediaOutput>())
+	{
+		return SAssignNew(PanelContent.Widget, SMediaProfileMediaOutputDisplay, SharedThis(this))
+			.MediaProfileEditor(MediaProfileEditor)
+			.PanelIndex(InPanelIndex)
+			.MediaItemIndex(PanelContent.MediaItemIndex);
+	}
+
+	return SAssignNew(PanelContent.Widget, SMediaProfileDummyDisplay, SharedThis(this))
+		.MediaProfileEditor(MediaProfileEditor)
+		.PanelIndex(InPanelIndex);
+}
+
+TSharedRef<SWidget> SMediaProfileViewport::ConstructViewportLayout()
+{
+	UMediaProfile* MediaProfile = nullptr;
+	if (MediaProfileEditor.IsValid())
+	{
+		MediaProfile = MediaProfileEditor.Pin()->GetMediaProfile();
+	}
+	
+	if (!MediaProfile)
+	{
+		return SNullWidget::NullWidget;
+	}
+	
+	if (CurrentLayout == EViewportLayout::OnePanel)
+	{
+		const bool bIsMaximizingPanel = MaximizedPanelIndex != INDEX_NONE && PanelContents.IsValidIndex(MaximizedPanelIndex);
+		return ConstructPanelWidget(bIsMaximizingPanel ? MaximizedPanelIndex : 0);
+	}
+
+	if (CurrentLayout == EViewportLayout::TwoPanel)
+	{
+		const bool bIsHorizontal = CurrentOrientation == EViewportOrientation::Left || CurrentOrientation == EViewportOrientation::Right;
+		return SNew(SSplitter)
+			.Orientation(bIsHorizontal ? Orient_Horizontal : Orient_Vertical)
+
+			+SSplitter::Slot()
+			[
+				ConstructPanelWidget(0)
+			]
+
+			+SSplitter::Slot()
+			[
+				ConstructPanelWidget(1)
+			];
+	}
+
+	if (CurrentLayout == EViewportLayout::ThreePanel)
+	{
+		const bool bIsHorizontal = CurrentOrientation == EViewportOrientation::Left || CurrentOrientation == EViewportOrientation::Right;
+		const bool bIsTopLeft = CurrentOrientation == EViewportOrientation::Left || CurrentOrientation == EViewportOrientation::Top;
+		TSharedRef<SSplitter> SecondarySplitter = SNew(SSplitter)
+			.Orientation(bIsHorizontal ? Orient_Vertical : Orient_Horizontal)
+			
+			+SSplitter::Slot()
+			[
+				ConstructPanelWidget(1)
+			]
+
+			+SSplitter::Slot()
+			[
+				ConstructPanelWidget(2)
+			];
+
+		return SNew(SSplitter)
+			.Orientation(bIsHorizontal ? Orient_Horizontal : Orient_Vertical)
+
+			+SSplitter::Slot()
+			[
+				bIsTopLeft ? ConstructPanelWidget(0) : SecondarySplitter
+			]
+
+			+SSplitter::Slot()
+			[
+				bIsTopLeft ? SecondarySplitter : ConstructPanelWidget(0)
+			];
+	}
+
+	if (CurrentLayout == EViewportLayout::OneAndThreePanel)
+	{
+		const bool bIsHorizontal = CurrentOrientation == EViewportOrientation::Left || CurrentOrientation == EViewportOrientation::Right;
+		const bool bIsTopLeft = CurrentOrientation == EViewportOrientation::Left || CurrentOrientation == EViewportOrientation::Top;
+		TSharedRef<SSplitter> SecondarySplitter = SNew(SSplitter)
+			.Orientation(bIsHorizontal ? Orient_Vertical : Orient_Horizontal)
+			
+			+SSplitter::Slot()
+			[
+				ConstructPanelWidget(1)
+			]
+
+			+SSplitter::Slot()
+			[
+				ConstructPanelWidget(2)
+			]
+
+			+SSplitter::Slot()
+			[
+				ConstructPanelWidget(3)
+			];;
+
+		return SNew(SSplitter)
+			.Orientation(bIsHorizontal ? Orient_Horizontal : Orient_Vertical)
+
+			+SSplitter::Slot()
+			[
+				bIsTopLeft ? ConstructPanelWidget(0) : SecondarySplitter
+			]
+
+			+SSplitter::Slot()
+			[
+				bIsTopLeft ? SecondarySplitter : ConstructPanelWidget(0)
+			];
+	}
+
+	if (CurrentLayout == EViewportLayout::QuadPanel)
+	{
+		return SNew(SSplitter2x2)
+			.TopLeft()
+			[
+				ConstructPanelWidget(0)
+			]
+			.BottomLeft()
+			[
+				ConstructPanelWidget(2)
+			]
+			.TopRight()
+			[
+				ConstructPanelWidget(1)
+			]
+			.BottomRight()
+			[
+				ConstructPanelWidget(3)
+			];
+	}
+
+	return SNullWidget::NullWidget;
+}
+
+void SMediaProfileViewport::SetViewportLayout(EViewportLayout InViewportLayout, EViewportOrientation InViewportOrientation)
+{
+	if (ImmersivePanelIndex != INDEX_NONE)
+	{
+		ClearImmersivePanel();
+	}
+	
+	CurrentLayout = InViewportLayout;
+	CurrentOrientation = InViewportOrientation;
+	MaximizedPanelIndex = INDEX_NONE;
+	
+	UpdateDisplays();
+}
+
+bool SMediaProfileViewport::IsViewportLayoutSet(EViewportLayout InViewportLayout, EViewportOrientation InViewportOrientation) const
+{
+	return CurrentLayout == InViewportLayout && CurrentOrientation == InViewportOrientation;
+}
+
+#undef LOCTEXT_NAMESPACE
\ No newline at end of file
diff --git a/Engine/Plugins/Media/MediaFrameworkUtilities/Source/MediaFrameworkUtilitiesEditor/Private/AssetEditor/SMediaProfileViewport.h b/Engine/Plugins/Media/MediaFrameworkUtilities/Source/MediaFrameworkUtilitiesEditor/Private/AssetEditor/SMediaProfileViewport.h
new file mode 100644
--- /dev/null
+++ b/Engine/Plugins/Media/MediaFrameworkUtilities/Source/MediaFrameworkUtilitiesEditor/Private/AssetEditor/SMediaProfileViewport.h
@@ -0,0 +1,154 @@
+﻿// Copyright Epic Games, Inc. All Rights Reserved.
+
+#pragma once
+
+#include "ToolMenuEntry.h"
+#include "Widgets/SCompoundWidget.h"
+
+class FMediaProfileEditor;
+class SOverlay;
+class FUICommandList;
+class UMediaProfile;
+
+/** Viewport that displays the input or output from selected media sources and outputs */
+class SMediaProfileViewport : public SCompoundWidget
+{
+public:
+	/** Indicates how to lay out the panels of the viewport */
+	enum class EViewportLayout
+	{
+		/** Single panel layout */
+		OnePanel,
+
+		/** Two panel layout split either horizontally or vertically */
+		TwoPanel,
+
+		/** Three panel layout where one half is a full panel and the other half is split between two panels */
+		ThreePanel,
+
+		/** Four panel layout where one half is a full panel and the other half is split between three panels */
+		OneAndThreePanel,
+
+		/** A 2x2 panel split */
+		QuadPanel
+	};
+
+	/** Indicates how to orient the primary panel in panel layouts that can have multiple orientations (e.g. three panel) */
+	enum class EViewportOrientation
+	{
+		Left,
+		Top,
+		Right,
+		Bottom
+	};
+	
+public:
+	SLATE_BEGIN_ARGS(SMediaProfileViewport) {}
+	SLATE_END_ARGS()
+		
+	void Construct(const FArguments& InArgs, const TSharedPtr<FMediaProfileEditor> InOwningEditor);
+	
+	virtual FReply OnKeyDown(const FGeometry& MyGeometry, const FKeyEvent& InKeyEvent) override;
+	virtual bool SupportsKeyboardFocus() const override { return true; }
+	
+	/** Displays the selected media items in a viewport panel arrangement */
+	void SetSelectedMediaItems(const TArray<int32>& InSelectedMediaSourceIndices, const TArray<int32>& InSelectedMediaOutputIndices);
+
+	/** Removes all panels that contain the specified media item, even if the panels are locked */
+	void ForceClearMediaItem(UClass* InMediaType, int32 InMediaItemIndex);
+
+	/** Sets the contents of the specified panel to the specified media type and media index */
+	void SetPanelContents(int32 InPanelIndex, UClass* InMediaItemClass, int32 InMediaItemIndex, bool bRefreshDisplay);
+	
+	/** Sets the lock status the specified panel */
+	void SetPanelLocked(int32 InPanelIndex, bool bInIsLocked);
+	
+	/** Gets whether the specified panel is locked or not */
+	bool IsPanelLocked(int32 InPanelIndex) const;
+	
+	/** Gets whether the specified panel can be maximized right now or not*/
+	bool CanMaximizePanel(int32 InPanelIndex) const;
+	
+	/** Maximizes the contents in the specified panel index, moving them to the first panel and changing the arrangement to single panel */
+	void MaximizePanel(int32 InPanelIndex);
+
+	/** Gets whether the specified panel is currently maximized */
+	bool IsPanelMaximized(int32 InPanelIndex) const;
+	
+	/** Gets whether there is a previous layout that can be restored */
+	bool CanRestorePreviousLayout() const;
+	
+	/** Restores the previous multi-panel layout from a maximized state, if one is saved */
+	void RestorePreviousLayout();
+
+	/** Makes the specified panel immersive */
+	void SetImmersivePanel(int32 InPanelIndex);
+
+	/** Gets whether the specified panel is being displayed in immersive view */
+	bool IsPanelImmersive(int32 InPanelIndex) const;
+	
+	/** Clears any immersive panel being displayed */
+	void ClearImmersivePanel();
+
+	/** Gets the viewport's command list */
+	TSharedPtr<FUICommandList> GetCommandList() { return CommandList; }
+	
+private:
+	/** Binds commands in this widget's command list to actions */
+	void BindCommands();
+	
+	/** Updates the current media item displays to match the selected media items  */
+	void UpdateDisplays();
+
+	/** Constructs the widget to display the contents of the specified panel */
+	TSharedRef<SWidget> ConstructPanelWidget(int32 InPanelIndex);
+
+	/** Creates and lays out the panels for the viewport based on the currently set layout, orientation, and panel contents */
+	TSharedRef<SWidget> ConstructViewportLayout();
+
+	/** Sets the viewport panel layout to the specified layout and orientation */
+	void SetViewportLayout(EViewportLayout InViewportLayout, EViewportOrientation InViewportOrientation);
+	
+	/** Returns whether the current viewport panel layout matches the specified layout and orientation */
+	bool IsViewportLayoutSet(EViewportLayout InViewportLayout, EViewportOrientation InViewportOrientation) const;
+	
+private:
+	/** Media profile editor that owns this viewport */
+	TWeakPtr<FMediaProfileEditor> MediaProfileEditor;
+
+	/** Container that holds the actual panel layout widgets */
+	TSharedPtr<SOverlay> LayoutContainer;
+
+	/** Stores the content of a displayed viewport panel */
+	struct FPanelContent
+	{
+		UClass* MediaType = nullptr;
+		int32 MediaItemIndex = INDEX_NONE;
+		bool bIsLocked = false;
+		TSharedPtr<SWidget> Widget;
+	};
+
+	/** The content of the currently displayed viewport panels */
+	TArray<FPanelContent> PanelContents;
+
+	/** Command list for the viewport tab */
+	TSharedPtr<FUICommandList> CommandList;
+
+	/** Currently active viewport layout */
+	EViewportLayout CurrentLayout = EViewportLayout::OnePanel;
+	
+	/** Currently active viewport orientation */
+	EViewportOrientation CurrentOrientation = EViewportOrientation::Left;
+	
+	/** Stores the previously used layout to restore if toggling out of maximized panel */
+	EViewportLayout SavedLayout = EViewportLayout::OnePanel;
+
+	/** Stores the previously used orientation to restore if toggling out of maximized panel */
+	EViewportOrientation SavedOrientation = EViewportOrientation::Left;
+
+	/** When maximizing from a multi-panel layout, this is the panel being maximized */
+	int32 MaximizedPanelIndex = INDEX_NONE;
+
+	/** When toggling immersive mode, this is the panel being displayed */
+	int32 ImmersivePanelIndex = INDEX_NONE;
+};
\ No newline at end of file
