Media Profile: Fixed LOCTEXT collision

diff --git a/Engine/Plugins/Media/MediaFrameworkUtilities/Source/MediaFrameworkUtilitiesEditor/Private/AssetEditor/SMediaProfileMediaItemDisplay.cpp b/Engine/Plugins/Media/MediaFrameworkUtilities/Source/MediaFrameworkUtilitiesEditor/Private/AssetEditor/SMediaProfileMediaItemDisplay.cpp
--- a/Engine/Plugins/Media/MediaFrameworkUtilities/Source/MediaFrameworkUtilitiesEditor/Private/AssetEditor/SMediaProfileMediaItemDisplay.cpp
+++ b/Engine/Plugins/Media/MediaFrameworkUtilities/Source/MediaFrameworkUtilitiesEditor/Private/AssetEditor/SMediaProfileMediaItemDisplay.cpp
@@ -732,14 +732,14 @@ void SMediaProfileDummyDisplay::ConfigureMediaImage()
 		.BorderImage(FAppStyle::GetBrush("NoBorder"))
 		[
 			SNew(STextBlock)
-			.Text(LOCTEXT("NoMediaItemSelectedLabel", "No media item selected"))
+			.Text(LOCTEXT("NoMediaItemSelectedLabel", "No Media Item selected"))
 		]
 	);
 }
 
 FText SMediaProfileDummyDisplay::GetActiveMediaItemLabel() const
 {
-	return LOCTEXT("NoMediaItemSelectedLabel", "No Media Item Selected");
+	return LOCTEXT("NoMediaItemSelectedLabel", "No Media Item selected");
 }
 
 void SMediaProfileDummyDisplay::ChangeActiveMediaItem(UClass* InMediaItemClass, int32 InMediaItemIndex)
