Media Profile: Added timecode/genlock details panel to the new media profile editor, as well as timecode and genlock widgets to the editor toolbar and viewport toolbar.

diff --git a/Engine/Plugins/Media/MediaFrameworkUtilities/Source/MediaFrameworkUtilities/Private/Profile/MediaProfile.cpp b/Engine/Plugins/Media/MediaFrameworkUtilities/Source/MediaFrameworkUtilities/Private/Profile/MediaProfile.cpp
--- a/Engine/Plugins/Media/MediaFrameworkUtilities/Source/MediaFrameworkUtilities/Private/Profile/MediaProfile.cpp
+++ b/Engine/Plugins/Media/MediaFrameworkUtilities/Source/MediaFrameworkUtilities/Private/Profile/MediaProfile.cpp
@@ -374,6 +374,45 @@
 	return bOverrideCustomTimeStep ? CustomTimeStep : nullptr;
 }
 
+void UMediaProfile::ApplyTimecodeProvider()
+{
+	ResetTimecodeProvider();
+	if (bOverrideTimecodeProvider)
+	{
+		bTimecodeProvideWasApplied = true;
+		AppliedTimecodeProvider = TimecodeProvider;
+		PreviousTimecodeProvider = GEngine->GetTimecodeProvider();
 
+		if (PreviousTimecodeProvider && AppliedTimecodeProvider && PreviousTimecodeProvider->GetClass() == AppliedTimecodeProvider->GetClass())
+		{
+			GEngine->SetTimecodeProvider(nullptr);
+			// Sleep to let the timecode provider complete its shutdown on its hardware.
+			FPlatformProcess::Sleep(0.2);
+		}
+
+		bool bResult = GEngine->SetTimecodeProvider(TimecodeProvider);
+		if (!bResult && TimecodeProvider)
+		{
+			UE_LOG(LogMediaFrameworkUtilities, Warning, TEXT("The Timecode Provider '%s' could not be initialized."), *TimecodeProvider->GetName());
+		}
+	}
+}
+
+void UMediaProfile::ApplyCustomTimeStep()
+{
+	ResetCustomTimeStep();
+	if (bOverrideCustomTimeStep)
+	{
+		bCustomTimeStepWasApplied = true;
+		AppliedCustomTimeStep = CustomTimeStep;
+		PreviousCustomTimeStep = GEngine->GetCustomTimeStep();
+		bool bResult = GEngine->SetCustomTimeStep(CustomTimeStep);
+		if (!bResult && CustomTimeStep)
+		{
+			UE_LOG(LogMediaFrameworkUtilities, Warning, TEXT("The Custom Time Step '%s' could not be initialized."), *CustomTimeStep->GetName());
+		}
+	}
+}
+
 void UMediaProfile::Apply()
 {
@@ -416,39 +455,8 @@
 		}
 	}
 
-	ResetTimecodeProvider();
-	if (bOverrideTimecodeProvider)
-	{
-		bTimecodeProvideWasApplied = true;
-		AppliedTimecodeProvider = TimecodeProvider;
-		PreviousTimecodeProvider = GEngine->GetTimecodeProvider();
+	ApplyTimecodeProvider();
+	ApplyCustomTimeStep();
 
-		if (PreviousTimecodeProvider && AppliedTimecodeProvider && PreviousTimecodeProvider->GetClass() == AppliedTimecodeProvider->GetClass())
-		{
-			GEngine->SetTimecodeProvider(nullptr);
-			// Sleep to let the timecode provider complete its shutdown on its hardware.
-			FPlatformProcess::Sleep(0.2);
-		}
-
-		bool bResult = GEngine->SetTimecodeProvider(TimecodeProvider);
-		if (!bResult && TimecodeProvider)
-		{
-			UE_LOG(LogMediaFrameworkUtilities, Warning, TEXT("The Timecode Provider '%s' could not be initialized."), *TimecodeProvider->GetName());
-		}
-	}
-
-	ResetCustomTimeStep();
-	if (bOverrideCustomTimeStep)
-	{
-		bCustomTimeStepWasApplied = true;
-		AppliedCustomTimeStep = CustomTimeStep;
-		PreviousCustomTimeStep = GEngine->GetCustomTimeStep();
-		bool bResult = GEngine->SetCustomTimeStep(CustomTimeStep);
-		if (!bResult && CustomTimeStep)
-		{
-			UE_LOG(LogMediaFrameworkUtilities, Warning, TEXT("The Custom Time Step '%s' could not be initialized."), *CustomTimeStep->GetName());
-		}
-	}
-
 	SendAnalytics();
 }
diff --git a/Engine/Plugins/Media/MediaFrameworkUtilities/Source/MediaFrameworkUtilities/Public/Profile/MediaProfile.h b/Engine/Plugins/Media/MediaFrameworkUtilities/Source/MediaFrameworkUtilities/Public/Profile/MediaProfile.h
--- a/Engine/Plugins/Media/MediaFrameworkUtilities/Source/MediaFrameworkUtilities/Public/Profile/MediaProfile.h
+++ b/Engine/Plugins/Media/MediaFrameworkUtilities/Source/MediaFrameworkUtilities/Public/Profile/MediaProfile.h
@@ -212,6 +212,12 @@ public:
 	 */
 	UEngineCustomTimeStep* GetCustomTimeStep() const;
 
+	/** Resets and applies the media profile's timecode to the engine if it is set to override the project settings timecode */
+	void ApplyTimecodeProvider();
+
+	/** Resets and applies the media profile's custom time step to the engine if it is set to override the project settings timecode */
+	void ApplyCustomTimeStep();
+	
 public:
 
 	/**
diff --git a/Engine/Plugins/Media/MediaFrameworkUtilities/Source/MediaFrameworkUtilitiesEditor/Private/AssetEditor/MediaProfileEditor.cpp b/Engine/Plugins/Media/MediaFrameworkUtilities/Source/MediaFrameworkUtilitiesEditor/Private/AssetEditor/MediaProfileEditor.cpp
--- a/Engine/Plugins/Media/MediaFrameworkUtilities/Source/MediaFrameworkUtilitiesEditor/Private/AssetEditor/MediaProfileEditor.cpp
+++ b/Engine/Plugins/Media/MediaFrameworkUtilities/Source/MediaFrameworkUtilitiesEditor/Private/AssetEditor/MediaProfileEditor.cpp
@@ -18,6 +18,9 @@
 #include "Profile/MediaProfilePlaybackManager.h"
 #include "Profile/MediaProfileSettings.h"
 #include "Slate/SceneViewport.h"
+#include "UI/MediaFrameworkUtilitiesEditorStyle.h"
+#include "UI/SMediaFrameworkTimecodeGenlockHeader.h"
+#include "UI/SMediaFrameworkTimecodeGenlockPanel.h"
 #include "Widgets/SViewport.h"
 #include "Widgets/Docking/SDockTab.h"
 
@@ -215,6 +218,9 @@ void FMediaProfileEditor::Initialize(const EToolkitMode::Type Mode, const TShare
 	const bool bCreateDefaultToolbar = true;
 	InitAssetEditor(Mode, InitToolkitHost, AppName, StandaloneDefaultLayout, bCreateDefaultStandaloneMenu, bCreateDefaultToolbar, InMediaProfile);
 
+	ExtendToolbar();
+	RegenerateMenusAndToolbars();
+	
 	// TODO: Not sure this is the best place for this, but for now, when we open a media profile, automatically make it the active media profile
 	// First, update the media profile settings config to match the media profile's source and output count
 	GetMutableDefault<UMediaProfileSettings>()->FillDefaultMediaSourceProxies(InMediaProfile->NumMediaSources());
@@ -256,7 +262,8 @@ void FMediaProfileEditor::RegisterTabSpawners(const TSharedRef<class FTabManager
 
 	InTabManager->RegisterTabSpawner(TimecodeTabId, FOnSpawnTab::CreateSP(this, &FMediaProfileEditor::SpawnTab_Timecode))
 		.SetDisplayName(LOCTEXT("TimecodeTabDisplayName", "Timecode/Genlock"))
-		.SetGroup(WorkspaceMenuCategory.ToSharedRef());
+		.SetGroup(WorkspaceMenuCategory.ToSharedRef())
+		.SetIcon(FSlateIcon(FMediaFrameworkUtilitiesEditorStyle::GetStyleSetName(), TEXT("ToolbarIcon.Timecode")));
 }
 
 void FMediaProfileEditor::UnregisterTabSpawners(const TSharedRef<class FTabManager>& InTabManager)
@@ -561,12 +568,31 @@ TSharedRef<SDockTab> FMediaProfileEditor::SpawnTab_Details(const FSpawnTabArgs&
 TSharedRef<SDockTab> FMediaProfileEditor::SpawnTab_Timecode(const FSpawnTabArgs& InArgs)
 {
 	return SNew(SDockTab)
-		.Label(LOCTEXT("TimecodeTabLabel", "Timecode/Genlock"))
+		.Label(LOCTEXT("TimecodeTabLabel", "TC/Genlock"))
+		.ToolTipText(LOCTEXT("TimecodeTabTooltip", "Timecode and Genlock settings"))
 		[
-			SNullWidget::NullWidget
+			SNew(SMediaFrameworkTimecodeGenlockPanel)
+				.MediaProfile(MediaProfileBeingEdited)
 		];
 }
 
+void FMediaProfileEditor::ExtendToolbar()
+{
+	TSharedPtr<FExtender> ToolbarExtender = MakeShared<FExtender>();
+
+	ToolbarExtender->AddToolBarExtension(
+		"Asset",
+		EExtensionHook::After,
+		GetToolkitCommands(),
+		FToolBarExtensionDelegate::CreateLambda([this](FToolBarBuilder& ToolbarBuilder)
+		{
+			ToolbarBuilder.AddWidget(SAssignNew(TimecodeToolbarEntry, SMediaFrameworkTimecodeGenlockHeader));
+		})
+	);
+
+	AddToolbarExtender(ToolbarExtender);
+}
+
 void FMediaProfileEditor::OnMediaItemDeleted(UClass* InMediaType, int32 InMediaItemIndex)
 {
 	if (InMediaType->IsChildOf<UMediaSource>())
diff --git a/Engine/Plugins/Media/MediaFrameworkUtilities/Source/MediaFrameworkUtilitiesEditor/Private/AssetEditor/MediaProfileEditor.h b/Engine/Plugins/Media/MediaFrameworkUtilities/Source/MediaFrameworkUtilitiesEditor/Private/AssetEditor/MediaProfileEditor.h
--- a/Engine/Plugins/Media/MediaFrameworkUtilities/Source/MediaFrameworkUtilitiesEditor/Private/AssetEditor/MediaProfileEditor.h
+++ b/Engine/Plugins/Media/MediaFrameworkUtilities/Source/MediaFrameworkUtilitiesEditor/Private/AssetEditor/MediaProfileEditor.h
@@ -5,6 +5,7 @@
 #include "CoreMinimal.h"
 #include "Toolkits/AssetEditorToolkit.h"
 
+class SMediaFrameworkTimecodeGenlockHeader;
 class SMediaProfileViewport;
 class FMediaFrameworkCaptureLevelEditorViewportClient;
 class UMediaFrameworkWorldSettingsAssetUserData;
@@ -99,6 +100,8 @@ private:
 	TSharedRef<SDockTab> SpawnTab_MediaTree(const FSpawnTabArgs& InArgs);
 	TSharedRef<SDockTab> SpawnTab_Details(const FSpawnTabArgs& InArgs);
 	TSharedRef<SDockTab> SpawnTab_Timecode(const FSpawnTabArgs& InArgs);
+
+	void ExtendToolbar();
 	
 	void OnMediaItemDeleted(UClass* InMediaType, int32 InMediaItemIndex);
 	void OnSelectedMediaItemsChanged(const TArray<int32>& SelectedMediaSources, const TArray<int32>& SelectedMediaOutputs);
@@ -124,6 +127,9 @@ private:
 
 	TSharedPtr<SMediaProfileDetailsPanel> DetailsPanel;
 	TSharedPtr<SMediaProfileViewport> ViewportPanel;
+	
+	/** Timecode/Genlock header widget that is displayed in the editor's primary toolbar */
+	TSharedPtr<SMediaFrameworkTimecodeGenlockHeader> TimecodeToolbarEntry;
 
 	/**
 	 * A map that stores the viewport captures for specific media capture camera viewport outputs.
diff --git a/Engine/Plugins/Media/MediaFrameworkUtilities/Source/MediaFrameworkUtilitiesEditor/Private/AssetEditor/MediaProfileEditorUserSettings.h b/Engine/Plugins/Media/MediaFrameworkUtilities/Source/MediaFrameworkUtilitiesEditor/Private/AssetEditor/MediaProfileEditorUserSettings.h
new file mode 100644
--- /dev/null
+++ b/Engine/Plugins/Media/MediaFrameworkUtilities/Source/MediaFrameworkUtilitiesEditor/Private/AssetEditor/MediaProfileEditorUserSettings.h
@@ -0,0 +1,23 @@
+﻿// Copyright Epic Games, Inc. All Rights Reserved.
+
+#pragma once
+
+#include "UObject/Object.h"
+
+#include "MediaProfileEditorUserSettings.generated.h"
+
+/** Config class to store editor user settings for the media profile editor */
+UCLASS(MinimalAPI, Config=EditorPerProjectUserSettings)
+class UMediaProfileEditorUserSettings : public UObject
+{
+	GENERATED_BODY()
+	
+public:
+	/** Whether to show the timecode entry in the viewport toolbar */
+	UPROPERTY(Config)
+	bool bShowTimecodeInViewportToolbar = true;
+
+	/** Whether to show the genlock entry in the viewport toolbar */
+	UPROPERTY(Config)
+	bool bShowGenlockInViewportToolbar = true;
+};
\ No newline at end of file
diff --git a/Engine/Plugins/Media/MediaFrameworkUtilities/Source/MediaFrameworkUtilitiesEditor/Private/AssetEditor/SMediaProfileMediaItemDisplay.h b/Engine/Plugins/Media/MediaFrameworkUtilities/Source/MediaFrameworkUtilitiesEditor/Private/AssetEditor/SMediaProfileMediaItemDisplay.h
--- a/Engine/Plugins/Media/MediaFrameworkUtilities/Source/MediaFrameworkUtilitiesEditor/Private/AssetEditor/SMediaProfileMediaItemDisplay.h
+++ b/Engine/Plugins/Media/MediaFrameworkUtilities/Source/MediaFrameworkUtilitiesEditor/Private/AssetEditor/SMediaProfileMediaItemDisplay.h
@@ -5,13 +5,17 @@
 #include "LevelViewportActions.h"
 #include "MediaOutput.h"
 #include "MediaProfileEditor.h"
+#include "MediaProfileEditorUserSettings.h"
 #include "MediaSource.h"
 #include "PropertyEditorModule.h"
 #include "SMediaProfileViewport.h"
 #include "ToolMenus.h"
+#include "ToolMenusEditor.h"
 #include "Framework/MultiBox/MultiBoxBuilder.h"
 #include "Profile/MediaProfile.h"
 #include "Styling/SlateIconFinder.h"
+#include "Tests/ToolMenusTestUtilities.h"
+#include "UI/MediaFrameworkTimecodeGenlockToolMenuEntry.h"
 #include "Widgets/SCompoundWidget.h"
 #include "Widgets/SOverlay.h"
 #include "Widgets/Layout/SBorder.h"
@@ -73,6 +77,8 @@ public:
 		MediaItemIndex = InArgs._MediaItemIndex;
 		PanelIndex = InArgs._PanelIndex;
 
+		TimecodeGenlockToolMenuEntry = MakeShared<FMediaFrameworkTimecodeGenlockToolMenuEntry>(GetMediaProfile());
+		
 		FCoreUObjectDelegates::OnObjectPropertyChanged.AddSP(this, &SMediaProfileMediaItemDisplayBase::OnObjectPropertyChanged);
 
 		CommandList = MakeShared<FUICommandList>();
@@ -93,12 +99,7 @@ public:
 					SNew(SHorizontalBox)
 
 					+SHorizontalBox::Slot()
-					[
-						SNew(SSpacer)
-					]
-
-					+SHorizontalBox::Slot()
-					.AutoWidth()
+					.FillWidth(1.0f)
 					[
 						CreateToolbar()
 					]
@@ -236,6 +237,47 @@ protected:
 		{
 			UToolMenu* ToolbarMenu = UToolMenus::Get()->RegisterMenu(ViewportToolbarName, NAME_None, EMultiBoxType::SlimHorizontalToolBar);
 			ToolbarMenu->StyleName = "ViewportToolbar";
+
+			// Need empty left section to ensure the middle and right sections get positioned correctly
+			FToolMenuSection& LeftSection = ToolbarMenu->AddSection("Left");
+			LeftSection.Alignment = EToolMenuSectionAlign::First;
+			
+			LeftSection.AddDynamicEntry(NAME_None, FNewToolMenuSectionDelegate::CreateLambda([](FToolMenuSection& Section)
+			{
+				UMediaProfileMediaItemDisplayContext* Context = Section.FindContext<UMediaProfileMediaItemDisplayContext>();
+				TSharedPtr<SMediaProfileMediaItemDisplayBase<TMediaItem>> ContextWidget = Context->GetDisplayWidget<SMediaProfileMediaItemDisplayBase<TMediaItem>>();
+				
+				if (!ContextWidget.IsValid())
+				{
+					return;
+				}
+
+				FToolMenuEntry TimecodeEntry = ContextWidget->TimecodeGenlockToolMenuEntry->CreateTimecodeToolMenuEntry(TAttribute<bool>::CreateLambda([]
+				{
+					if (const UMediaProfileEditorUserSettings* UserSettings = GetDefault<UMediaProfileEditorUserSettings>())
+					{
+						return UserSettings->bShowTimecodeInViewportToolbar;
+					}
+
+					return true;
+				}));
+
+				//TODO: Hidden for now until per-source timecode can be figured out (see UE-UE-305891)
+				//Section.AddEntry(TimecodeEntry);
+				
+				FToolMenuEntry GenlockEntry = ContextWidget->TimecodeGenlockToolMenuEntry->CreateGenlockToolMenuEntry(TAttribute<bool>::CreateLambda([]
+				{
+					if (const UMediaProfileEditorUserSettings* UserSettings = GetDefault<UMediaProfileEditorUserSettings>())
+					{
+						return UserSettings->bShowGenlockInViewportToolbar;
+					}
+
+					return true;
+				}));
+
+				//TODO: Hidden for now until per-source genlock can be figured out (see UE-UE-305891)
+				//Section.AddEntry(GenlockEntry);
+			}));
 			
 			FToolMenuSection& RightSection = ToolbarMenu->AddSection("Right");
 			RightSection.Alignment = EToolMenuSectionAlign::Last;
@@ -253,6 +295,9 @@ protected:
 				Section.AddEntry(ContextWidget->CreateMediaItemsComboButton());
 				ContextWidget->AddToolbarEntries(Section);
 				Section.AddEntry(ContextWidget->CreateViewportArrangementMenu());
+
+				//TODO: Hidden for now until per-source timecode/genlock can be figured out (see UE-UE-305891)
+				//Section.AddEntry(ContextWidget->CreateSettingsMenu());
 			}));
 		}
 		
@@ -455,6 +500,74 @@ protected:
 		return MaximizeRestoreEntry;
 	}
 
+	FToolMenuEntry CreateSettingsMenu()
+	{
+		FToolMenuEntry Entry = FToolMenuEntry::InitSubMenu(
+			"SettingsMenu",
+			TAttribute<FText>(),
+			TAttribute<FText>(),
+			FNewToolMenuDelegate::CreateLambda([this](UToolMenu* Submenu)
+			{
+				FToolMenuSection& TimecodeGenlockSection = Submenu->AddSection("TimecodeGenlockSettingsSection", LOCTEXT("TimecodeGenlockSettingsSectionLabel", "Timecode & Genlock"));
+				TimecodeGenlockSection.AddMenuEntry(
+					"ShowTimecodeInToolbar",
+					LOCTEXT("ShowTimecodeInToolbarLabel", "Show Timecode in Viewport Toolbar"),
+					LOCTEXT("ShowTimecodeInToolbarTooltip", "Whether to show the timecode in the viewport toolbar"),
+					FSlateIcon(),
+					FUIAction(
+						FExecuteAction::CreateLambda([]
+						{
+							if (UMediaProfileEditorUserSettings* UserSettings = GetMutableDefault<UMediaProfileEditorUserSettings>())
+							{
+								UserSettings->bShowTimecodeInViewportToolbar = !UserSettings->bShowTimecodeInViewportToolbar;
+								UserSettings->SaveConfig();
+							}
+						}),
+						FCanExecuteAction(),
+						FIsActionChecked::CreateLambda([]
+						{
+							if (const UMediaProfileEditorUserSettings* UserSettings = GetDefault<UMediaProfileEditorUserSettings>())
+							{
+								return UserSettings->bShowTimecodeInViewportToolbar;
+							}
+
+							return true;
+						})),
+						EUserInterfaceActionType::ToggleButton);
+				
+				TimecodeGenlockSection.AddMenuEntry(
+					"ShowGenlockInToolbar",
+					LOCTEXT("ShowGenlockInToolbarLabel", "Show Genlock in Viewport Toolbar"),
+					LOCTEXT("ShowGenlockInToolbarTooltip", "Whether to show the genlock status in the viewport toolbar"),
+					FSlateIcon(),
+					FUIAction(
+						FExecuteAction::CreateLambda([]
+						{
+							if (UMediaProfileEditorUserSettings* UserSettings = GetMutableDefault<UMediaProfileEditorUserSettings>())
+							{
+								UserSettings->bShowGenlockInViewportToolbar = !UserSettings->bShowGenlockInViewportToolbar;
+								UserSettings->SaveConfig();
+							}
+						}),
+						FCanExecuteAction(),
+						FIsActionChecked::CreateLambda([]
+						{
+							if (const UMediaProfileEditorUserSettings* UserSettings = GetDefault<UMediaProfileEditorUserSettings>())
+							{
+								return UserSettings->bShowGenlockInViewportToolbar;
+							}
+
+							return true;
+						})),
+						EUserInterfaceActionType::ToggleButton);
+			}),
+			false,
+			FSlateIcon(FAppStyle::GetAppStyleSetName(), "LevelEditor.GameSettings")
+		);
+
+		return Entry;
+	}
+	
 	/** Gets the label to display for this panel's content in the media items combo button */
 	virtual FText GetActiveMediaItemLabel() const
 	{
@@ -581,6 +694,9 @@ protected:
 	/** The index of the media item being displayed by this display */
 	int32 MediaItemIndex = INDEX_NONE;
 
+	/** Entry for the media profile's timecode and genlock configuration to display in the viewport toolbar */
+	TSharedPtr<FMediaFrameworkTimecodeGenlockToolMenuEntry> TimecodeGenlockToolMenuEntry;
+	
 	TSharedPtr<SOverlay> Overlay;
 	TSharedPtr<SBorder> MediaImageContainer;
 	TSharedPtr<FUICommandList> CommandList;
diff --git a/Engine/Plugins/Media/MediaFrameworkUtilities/Source/MediaFrameworkUtilitiesEditor/Private/MediaFrameworkUtilitiesEditorModule.cpp b/Engine/Plugins/Media/MediaFrameworkUtilities/Source/MediaFrameworkUtilitiesEditor/Private/MediaFrameworkUtilitiesEditorModule.cpp
--- a/Engine/Plugins/Media/MediaFrameworkUtilities/Source/MediaFrameworkUtilitiesEditor/Private/MediaFrameworkUtilitiesEditorModule.cpp
+++ b/Engine/Plugins/Media/MediaFrameworkUtilities/Source/MediaFrameworkUtilitiesEditor/Private/MediaFrameworkUtilitiesEditorModule.cpp
@@ -75,6 +75,18 @@ public:
 			PropertyModule.RegisterCustomClassLayout(UMediaProfile::StaticClass()->GetFName(), FOnGetDetailCustomizationInstance::CreateStatic(&FMediaProfileCustomization::MakeInstance));
 			PropertyModule.RegisterCustomClassLayout(UMediaProfileSettings::StaticClass()->GetFName(), FOnGetDetailCustomizationInstance::CreateStatic(&FMediaProfileSettingsCustomization::MakeInstance));
 
+			RegisterSection(PropertyModule,
+				UMediaProfile::StaticClass(),
+				TEXT("TimecodeProvider"),
+				LOCTEXT("TimecodeProviderSection", "Timecode Provider"),
+				TArray<FName> { TEXT("Timecode Provider") });
+
+			RegisterSection(PropertyModule,
+				UMediaProfile::StaticClass(),
+				TEXT("Genlock"),
+				LOCTEXT("GenlockSection", "Genlock"),
+				TArray<FName> { TEXT("Genlock") });
+			
 			{
 				const IWorkspaceMenuStructure& MenuStructure = WorkspaceMenu::GetMenuStructure();
 				TSharedRef<FWorkspaceItem> MediaBrowserGroup = MenuStructure.GetLevelEditorVirtualProductionCategory()->AddGroup(
@@ -147,6 +159,16 @@ public:
 		}
 	}
 
+private:
+	void RegisterSection(FPropertyEditorModule& PropertyModule, UClass* InClass, FName InSectionName, FText InSectionLabel, const TArray<FName>& InCategoriesInSection)
+	{
+		const TSharedPtr<FPropertySection> Section = PropertyModule.FindOrCreateSection(InClass->GetFName(), InSectionName,InSectionLabel);
+		for (const FName& Category : InCategoriesInSection)
+		{
+			Section->AddCategory(Category);
+		}
+	}
+	
 private:
 
 	TArray<TSharedRef<IAssetTypeActions>> RegisteredAssetTypeActions;
diff --git a/Engine/Plugins/Media/MediaFrameworkUtilities/Source/MediaFrameworkUtilitiesEditor/Private/UI/MediaFrameworkTimecodeGenlockToolMenuEntry.cpp b/Engine/Plugins/Media/MediaFrameworkUtilities/Source/MediaFrameworkUtilitiesEditor/Private/UI/MediaFrameworkTimecodeGenlockToolMenuEntry.cpp
new file mode 100644
--- /dev/null
+++ b/Engine/Plugins/Media/MediaFrameworkUtilities/Source/MediaFrameworkUtilitiesEditor/Private/UI/MediaFrameworkTimecodeGenlockToolMenuEntry.cpp
@@ -0,0 +1,142 @@
+﻿// Copyright Epic Games, Inc. All Rights Reserved.
+
+
+#include "MediaFrameworkTimecodeGenlockToolMenuEntry.h"
+
+#include "MediaFrameworkUtilitiesEditorStyle.h"
+#include "ToolMenu.h"
+#include "Engine/EngineCustomTimeStep.h"
+#include "Engine/TimecodeProvider.h"
+#include "Misc/Timecode.h"
+#include "Profile/MediaProfile.h"
+#include "Styling/SlateIconFinder.h"
+
+class UTimecodeProvider;
+
+#define LOCTEXT_NAMESPACE "FMediaFrameworkTimecodeGenlockToolMenuEntry"
+
+FToolMenuEntry FMediaFrameworkTimecodeGenlockToolMenuEntry::CreateTimecodeToolMenuEntry(const TAttribute<bool>& InEntryVisibleAttribute)
+{
+	FToolMenuEntry Entry = FToolMenuEntry::InitSubMenu(
+		TEXT("ActiveTimecodeEntry"),
+		TAttribute<FText>::CreateSP(this, &FMediaFrameworkTimecodeGenlockToolMenuEntry::GetTimecodeEntryText),
+		LOCTEXT("ActiveTimecodeEntryTooltip", "Shows the active timecode"),
+		FNewToolMenuDelegate::CreateSP(this, &FMediaFrameworkTimecodeGenlockToolMenuEntry::GetTimecodeDropdownContent),
+		false,
+		FSlateIcon(FMediaFrameworkUtilitiesEditorStyle::GetStyleSetName(), TEXT("ToolbarIcon.Timecode")));
+
+	Entry.Visibility = InEntryVisibleAttribute;
+	
+	return Entry;
+}
+
+FToolMenuEntry FMediaFrameworkTimecodeGenlockToolMenuEntry::CreateGenlockToolMenuEntry(const TAttribute<bool>& InEntryVisibleAttribute)
+{
+	FToolMenuEntry Entry = FToolMenuEntry::InitSubMenu(
+		TEXT("ActiveGenlockEntry"),
+		TAttribute<FText>(),
+		LOCTEXT("ActiveGenlockEntryTooltip", "Shows if there is an active custom time step and if it is currently synchronized"),
+		FNewToolMenuDelegate::CreateSP(this, &FMediaFrameworkTimecodeGenlockToolMenuEntry::GetGenlockDropdownContent),
+		false,
+		TAttribute<FSlateIcon>::CreateSP(this, &FMediaFrameworkTimecodeGenlockToolMenuEntry::GetGenlockEntryIcon));
+
+	Entry.Visibility = InEntryVisibleAttribute;
+	
+	return Entry;
+}
+
+FText FMediaFrameworkTimecodeGenlockToolMenuEntry::GetTimecodeEntryText() const
+{
+	const FText NullTimecodeText = FText::FromString(FTimecode().ToString());
+	
+	TStrongObjectPtr<UMediaProfile> PinnedMediaProfile = MediaProfile.Pin();
+	if (!PinnedMediaProfile.IsValid())
+	{
+		return NullTimecodeText;
+	}
+
+	UTimecodeProvider* TimecodeProvider = PinnedMediaProfile->GetTimecodeProvider();
+	if (!TimecodeProvider)
+	{
+		return NullTimecodeText;
+	}
+
+	const ETimecodeProviderSynchronizationState SyncState = TimecodeProvider->GetSynchronizationState();
+	if (SyncState != ETimecodeProviderSynchronizationState::Synchronized && SyncState != ETimecodeProviderSynchronizationState::Synchronizing)
+	{
+		return NullTimecodeText;
+	}
+
+	return FText::FromString(TimecodeProvider->GetTimecode().ToString());
+}
+
+FSlateIcon FMediaFrameworkTimecodeGenlockToolMenuEntry::GetGenlockEntryIcon() const
+{
+	TStrongObjectPtr<UMediaProfile> PinnedMediaProfile = MediaProfile.Pin();
+	if (!PinnedMediaProfile.IsValid())
+	{
+		return FSlateIcon(FAppStyle::GetAppStyleSetName(), TEXT("Icons.WarningWithColor"));
+	}
+
+	UEngineCustomTimeStep* CustomTimeStep = PinnedMediaProfile->GetCustomTimeStep();
+	if (!CustomTimeStep)
+	{
+		return FSlateIcon(FAppStyle::GetAppStyleSetName(), TEXT("Icons.WarningWithColor"));
+	}
+
+	const ECustomTimeStepSynchronizationState SyncState = CustomTimeStep->GetSynchronizationState();
+	if (SyncState != ECustomTimeStepSynchronizationState::Synchronized && SyncState != ECustomTimeStepSynchronizationState::Synchronizing)
+	{
+		return FSlateIcon(FAppStyle::GetAppStyleSetName(), TEXT("Icons.WarningWithColor"));
+	}
+	
+	return FSlateIcon(FMediaFrameworkUtilitiesEditorStyle::GetStyleSetName(), TEXT("ToolbarIcon.Genlock"));
+}
+
+void FMediaFrameworkTimecodeGenlockToolMenuEntry::GetTimecodeDropdownContent(UToolMenu* ToolMenu)
+{
+	FText TimecodeDisplayText = LOCTEXT("TimecodeProviderNotConfiguredText", "Timecode Provider not configured");
+	FSlateIcon TimecodeIcon = FSlateIcon();
+	
+	TStrongObjectPtr<UMediaProfile> PinnedMediaProfile = MediaProfile.Pin();
+	if (PinnedMediaProfile.IsValid())
+	{
+		if (UTimecodeProvider* TimecodeProvider = PinnedMediaProfile->GetTimecodeProvider())
+		{
+			TimecodeDisplayText = TimecodeProvider->GetClass()->GetDisplayNameText();
+			TimecodeIcon = FSlateIconFinder::FindIconForClass(TimecodeProvider->GetClass());
+		}
+	}
+
+	FToolMenuSection& TimecodeProviderSection = ToolMenu->AddSection(TEXT("TimecodeProvider"), LOCTEXT("TimecodeProviderSectionLabel", "Timecode Provider"));
+	TimecodeProviderSection.AddMenuEntry("TimecodeProvider",
+		TimecodeDisplayText,
+		TAttribute<FText>(),
+		TimecodeIcon,
+		FUIAction(FExecuteAction(), FCanExecuteAction::CreateLambda([] { return false; })));
+}
+
+void FMediaFrameworkTimecodeGenlockToolMenuEntry::GetGenlockDropdownContent(UToolMenu* ToolMenu)
+{
+	FText GenlockDisplayText = LOCTEXT("GenlockNotConfiguredText", "Genlock not configured");
+	FSlateIcon GenlockIcon = FSlateIcon();
+	
+	TStrongObjectPtr<UMediaProfile> PinnedMediaProfile = MediaProfile.Pin();
+	if (PinnedMediaProfile.IsValid())
+	{
+		if (UEngineCustomTimeStep* CustomTimeStep = PinnedMediaProfile->GetCustomTimeStep())
+		{
+			GenlockDisplayText = FText::FromString(CustomTimeStep->GetDisplayName());
+			GenlockIcon = FSlateIconFinder::FindIconForClass(CustomTimeStep->GetClass());
+		}
+	}
+
+	FToolMenuSection& GenlockSection = ToolMenu->AddSection(TEXT("GenlockSection"), LOCTEXT("GenlockSectionLabel", "Genlock"));
+	GenlockSection.AddMenuEntry("Genlock",
+		GenlockDisplayText,
+		TAttribute<FText>(),
+		GenlockIcon,
+		FUIAction(FExecuteAction(), FCanExecuteAction::CreateLambda([] { return false; })));
+}
+
+#undef LOCTEXT_NAMESPACE
diff --git a/Engine/Plugins/Media/MediaFrameworkUtilities/Source/MediaFrameworkUtilitiesEditor/Private/UI/MediaFrameworkTimecodeGenlockToolMenuEntry.h b/Engine/Plugins/Media/MediaFrameworkUtilities/Source/MediaFrameworkUtilitiesEditor/Private/UI/MediaFrameworkTimecodeGenlockToolMenuEntry.h
new file mode 100644
--- /dev/null
+++ b/Engine/Plugins/Media/MediaFrameworkUtilities/Source/MediaFrameworkUtilitiesEditor/Private/UI/MediaFrameworkTimecodeGenlockToolMenuEntry.h
@@ -0,0 +1,42 @@
+﻿// Copyright Epic Games, Inc. All Rights Reserved.
+
+#pragma once
+
+#include "ToolMenuEntry.h"
+#include "UObject/WeakObjectPtrTemplates.h"
+
+class UMediaProfile;
+
+/**
+ * Creates and manages tool menu entries for timecode and genlock configurations on media profiles
+ */
+class FMediaFrameworkTimecodeGenlockToolMenuEntry: public TSharedFromThis<FMediaFrameworkTimecodeGenlockToolMenuEntry>
+{
+public:
+	FMediaFrameworkTimecodeGenlockToolMenuEntry(TWeakObjectPtr<UMediaProfile> InMediaProfile)
+		: MediaProfile(InMediaProfile)
+	{ }
+	
+	/** Creates the timecode tool menu entry that can be added to toolbars  */
+	FToolMenuEntry CreateTimecodeToolMenuEntry(const TAttribute<bool>& InEntryVisibleAttribute);
+
+	/** Creates the genlock tool menu entry that can be added to toolbars  */
+	FToolMenuEntry CreateGenlockToolMenuEntry(const TAttribute<bool>& InEntryVisibleAttribute);
+
+private:
+	/** Gets the text for the timecode toolbar entry */
+	FText GetTimecodeEntryText() const;
+
+	/** Gets the icon for the genlock toolbar entry */
+	FSlateIcon GetGenlockEntryIcon() const;
+
+	/** Fills the dropdown menu for the timecode toolbar entry */
+	void GetTimecodeDropdownContent(UToolMenu* ToolMenu);
+
+	/** Fills the dropdown menu for the genlock toolbar entry */
+	void GetGenlockDropdownContent(UToolMenu* ToolMenu);
+	
+private:
+	/** Weak pointer to the media profile to source the timecode and genlock configuration from */
+	TWeakObjectPtr<UMediaProfile> MediaProfile;
+};
diff --git a/Engine/Plugins/Media/MediaFrameworkUtilities/Source/MediaFrameworkUtilitiesEditor/Private/UI/MediaFrameworkUtilitiesEditorStyle.cpp b/Engine/Plugins/Media/MediaFrameworkUtilities/Source/MediaFrameworkUtilitiesEditor/Private/UI/MediaFrameworkUtilitiesEditorStyle.cpp
--- a/Engine/Plugins/Media/MediaFrameworkUtilities/Source/MediaFrameworkUtilitiesEditor/Private/UI/MediaFrameworkUtilitiesEditorStyle.cpp
+++ b/Engine/Plugins/Media/MediaFrameworkUtilities/Source/MediaFrameworkUtilitiesEditor/Private/UI/MediaFrameworkUtilitiesEditorStyle.cpp
@@ -51,6 +51,9 @@ void FMediaFrameworkUtilitiesEditorStyle::Register()
 	MediaFrameworkUtilitiesStyle::StyleInstance->Set("VideoInput.Stop", new IMAGE_BRUSH("CaptureCameraViewport_Stop_40x", MediaFrameworkUtilitiesStyle::Icon40x40));
 	MediaFrameworkUtilitiesStyle::StyleInstance->Set("VideoInput.Settings", new IMAGE_BRUSH("Settings_40x", MediaFrameworkUtilitiesStyle::Icon40x40));
 
+	MediaFrameworkUtilitiesStyle::StyleInstance->Set("ToolbarIcon.Timecode", new IMAGE_BRUSH_SVG("Timecode", MediaFrameworkUtilitiesStyle::Icon16x16));
+	MediaFrameworkUtilitiesStyle::StyleInstance->Set("ToolbarIcon.Genlock", new IMAGE_BRUSH_SVG("Genlock", MediaFrameworkUtilitiesStyle::Icon16x16));
+	
 	//Use existing icon of MaterialInstanceConstant from Engine Editor content
 	const FString EngineContentPath = FPaths::EngineContentDir() / TEXT("Editor/Slate/Icons/AssetIcons/MaterialInstanceConstant_64x.png");
 	MediaFrameworkUtilitiesStyle::StyleInstance->Set("MaterialEditor", new FSlateImageBrush(EngineContentPath, MediaFrameworkUtilitiesStyle::Icon40x40));
diff --git a/Engine/Plugins/Media/MediaFrameworkUtilities/Source/MediaFrameworkUtilitiesEditor/Private/UI/SMediaFrameworkTimecodeGenlockHeader.cpp b/Engine/Plugins/Media/MediaFrameworkUtilities/Source/MediaFrameworkUtilitiesEditor/Private/UI/SMediaFrameworkTimecodeGenlockHeader.cpp
new file mode 100644
--- /dev/null
+++ b/Engine/Plugins/Media/MediaFrameworkUtilities/Source/MediaFrameworkUtilitiesEditor/Private/UI/SMediaFrameworkTimecodeGenlockHeader.cpp
@@ -0,0 +1,140 @@
+﻿// Copyright Epic Games, Inc. All Rights Reserved.
+
+#include "SMediaFrameworkTimecodeGenlockHeader.h"
+
+#include "Components/HorizontalBox.h"
+#include "Engine/Engine.h"
+#include "Engine/EngineCustomTimeStep.h"
+#include "Engine/TimecodeProvider.h"
+#include "Profile/IMediaProfileManager.h"
+#include "Profile/MediaProfile.h"
+#include "Widgets/Images/SImage.h"
+#include "Widgets/Layout/SSeparator.h"
+#include "Widgets/Text/STextBlock.h"
+
+#define LOCTEXT_NAMESPACE "SMediaFrameworkTimecodeGenlockHeader"
+
+void SMediaFrameworkTimecodeGenlockHeader::Construct(const FArguments& InArgs)
+{
+	ChildSlot
+	[
+		SNew(SHorizontalBox)
+		.ToolTipText(this, &SMediaFrameworkTimecodeGenlockHeader::GetTooltipText)
+
+		+SHorizontalBox::Slot()
+		.AutoWidth()
+		.HAlign(HAlign_Left)
+		.VAlign(VAlign_Center)
+		.Padding(8.0f)
+		[
+			SNew(STextBlock)
+			.Text(this, &SMediaFrameworkTimecodeGenlockHeader::GetTimecodeText)
+			.Font(FCoreStyle::GetDefaultFontStyle("Bold", 10))
+		]
+		
+		+SHorizontalBox::Slot()
+		.AutoWidth()
+		[
+			SNew(SSeparator)
+			.Orientation(EOrientation::Orient_Vertical)
+			.Thickness(2)
+		]
+
+		+SHorizontalBox::Slot()
+		.AutoWidth()
+		.HAlign(HAlign_Center)
+		.VAlign(VAlign_Center)
+		.Padding(8.0)
+		[
+			SNew(SImage)
+			.Image(this, &SMediaFrameworkTimecodeGenlockHeader::GetGenlockIcon)
+			.DesiredSizeOverride(FVector2D(20, 20))
+		]
+		
+		+SHorizontalBox::Slot()
+		.AutoWidth()
+		.HAlign(HAlign_Left)
+		.VAlign(VAlign_Center)
+		.Padding(8.0f)
+		[
+			SNew(STextBlock)
+			.Text(this, &SMediaFrameworkTimecodeGenlockHeader::GetGenlockText)
+			.Font(FCoreStyle::GetDefaultFontStyle("Bold", 10))
+		]
+	];
+}
+
+FText SMediaFrameworkTimecodeGenlockHeader::GetTooltipText() const
+{
+	UTimecodeProvider* TimecodeProvider = GEngine->GetTimecodeProvider();
+	if (!TimecodeProvider)
+	{
+		return LOCTEXT("NoTimecodeProviderTextFormat", "Timecode not configured");
+	}
+	
+	FText SourceName = LOCTEXT("EngineSourceText", "Engine");
+	FText SourceType = LOCTEXT("EngineSourceTypeText", "Engine");
+	
+	if (UMediaProfile* ActiveMediaProfile = IMediaProfileManager::Get().GetCurrentMediaProfile())
+	{
+		if (ActiveMediaProfile->GetTimecodeProvider() == TimecodeProvider)
+		{
+			SourceName = FText::FromString(ActiveMediaProfile->GetName());
+			SourceType = LOCTEXT("MediaProfileTypeText", "Media Profile");
+		}
+	}
+	
+	const FText TimecodeProviderName = TimecodeProvider->GetClass()->GetDisplayNameText();
+	return FText::Format(LOCTEXT("TooltipTextFormat", "Timecode Provider: {0}\nSet by: {1} ({2})"), TimecodeProviderName, SourceName, SourceType);
+}
+
+FText SMediaFrameworkTimecodeGenlockHeader::GetTimecodeText() const
+{
+	const FText NullTimecodeText = FText::FromString(FTimecode().ToString());
+
+	UTimecodeProvider* TimecodeProvider = GEngine->GetTimecodeProvider();
+	if (!TimecodeProvider)
+	{
+		return NullTimecodeText;
+	}
+
+	const ETimecodeProviderSynchronizationState SyncState = TimecodeProvider->GetSynchronizationState();
+	if (SyncState != ETimecodeProviderSynchronizationState::Synchronized && SyncState != ETimecodeProviderSynchronizationState::Synchronizing)
+	{
+		return NullTimecodeText;
+	}
+ 
+	return FText::FromString(TimecodeProvider->GetTimecode().ToString());
+}
+
+const FSlateBrush* SMediaFrameworkTimecodeGenlockHeader::GetGenlockIcon() const
+{
+	const FSlateBrush* WarningIcon = FAppStyle::GetBrush("Icons.WarningWithColor");
+	
+	UEngineCustomTimeStep* CustomTimeStep = GEngine->GetCustomTimeStep();
+	if (!CustomTimeStep)
+	{
+		return WarningIcon;
+	}
+
+	const ECustomTimeStepSynchronizationState SyncState = CustomTimeStep->GetSynchronizationState();
+	if (SyncState != ECustomTimeStepSynchronizationState::Synchronized && SyncState != ECustomTimeStepSynchronizationState::Synchronizing)
+	{
+		return WarningIcon;
+	}
+ 
+	return nullptr;
+}
+
+FText SMediaFrameworkTimecodeGenlockHeader::GetGenlockText() const
+{
+	UEngineCustomTimeStep* CustomTimeStep = GEngine->GetCustomTimeStep();
+	if (!CustomTimeStep)
+	{
+		return LOCTEXT("NoCustomTimeStepLabel", "No custom time step configured");
+	}
+ 
+	return FText::FromString(CustomTimeStep->GetDisplayName());
+}
+
+#undef LOCTEXT_NAMESPACE
\ No newline at end of file
diff --git a/Engine/Plugins/Media/MediaFrameworkUtilities/Source/MediaFrameworkUtilitiesEditor/Private/UI/SMediaFrameworkTimecodeGenlockHeader.h b/Engine/Plugins/Media/MediaFrameworkUtilities/Source/MediaFrameworkUtilitiesEditor/Private/UI/SMediaFrameworkTimecodeGenlockHeader.h
new file mode 100644
--- /dev/null
+++ b/Engine/Plugins/Media/MediaFrameworkUtilities/Source/MediaFrameworkUtilitiesEditor/Private/UI/SMediaFrameworkTimecodeGenlockHeader.h
@@ -0,0 +1,33 @@
+﻿// Copyright Epic Games, Inc. All Rights Reserved.
+
+#pragma once
+
+#include "Widgets/SCompoundWidget.h"
+
+class UMediaProfile;
+
+/**
+ * Widget that displays the timecode and genlock status for a media profile that can be displayed in toolbars or details panel headers
+ */
+class SMediaFrameworkTimecodeGenlockHeader : public SCompoundWidget
+{
+public:
+	SLATE_BEGIN_ARGS(SMediaFrameworkTimecodeGenlockHeader) { }
+	SLATE_END_ARGS()
+
+	/** Constructs this widget with InArgs */
+	void Construct(const FArguments& InArgs);
+
+private:
+	/** Gets the text to display in the tooltip for the widget */
+	FText GetTooltipText() const;
+
+	/** Gets the text to display for the timecode */
+	FText GetTimecodeText() const;
+
+	/** Gets the icon to display for the genlock */
+	const FSlateBrush* GetGenlockIcon() const;
+
+	/** Gets the text to display for the genlock */
+	FText GetGenlockText() const;
+};
diff --git a/Engine/Plugins/Media/MediaFrameworkUtilities/Source/MediaFrameworkUtilitiesEditor/Private/UI/SMediaFrameworkTimecodeGenlockPanel.cpp b/Engine/Plugins/Media/MediaFrameworkUtilities/Source/MediaFrameworkUtilitiesEditor/Private/UI/SMediaFrameworkTimecodeGenlockPanel.cpp
new file mode 100644
--- /dev/null
+++ b/Engine/Plugins/Media/MediaFrameworkUtilities/Source/MediaFrameworkUtilitiesEditor/Private/UI/SMediaFrameworkTimecodeGenlockPanel.cpp
@@ -0,0 +1,109 @@
+﻿// Copyright Epic Games, Inc. All Rights Reserved.
+
+#include "SMediaFrameworkTimecodeGenlockPanel.h"
+
+#include "DetailLayoutBuilder.h"
+#include "IDetailCustomization.h"
+#include "SMediaFrameworkTimecodeGenlockHeader.h"
+#include "Engine/EngineCustomTimeStep.h"
+#include "Engine/TimecodeProvider.h"
+#include "Misc/Timecode.h"
+#include "Modules/ModuleManager.h"
+#include "Profile/MediaProfile.h"
+#include "Widgets/Images/SImage.h"
+#include "Widgets/Layout/SSeparator.h"
+#include "Widgets/Text/STextBlock.h"
+
+#define LOCTEXT_NAMESPACE "SMediaFrameworkTimecodeGenlockPanel"
+
+/** Details panel customization to hide any categories unrelated to timecode or genlock within the media profile  */
+class FMediaProfileTimecodeGenlockDetailsCustomization : public IDetailCustomization
+{
+	virtual void CustomizeDetails(IDetailLayoutBuilder& DetailBuilder) override
+	{
+		DetailBuilder.HideCategory(TEXT("Inputs"));
+		DetailBuilder.HideCategory(TEXT("Outputs"));
+	}
+};
+
+void SMediaFrameworkTimecodeGenlockPanel::Construct(const FArguments& InArgs)
+{
+	MediaProfile = InArgs._MediaProfile;
+	
+	FPropertyEditorModule& PropertyModule = FModuleManager::LoadModuleChecked<FPropertyEditorModule>("PropertyEditor");
+	FDetailsViewArgs DetailsArgs;
+	DetailsArgs.NameAreaSettings = FDetailsViewArgs::HideNameArea;
+	DetailsArgs.bCustomNameAreaLocation = true;
+	DetailsArgs.bCustomFilterAreaLocation = true;
+	DetailsArgs.DefaultsOnlyVisibility = EEditDefaultsOnlyNodeVisibility::Hide;
+	DetailsArgs.bShowSectionSelector = true;
+	DetailsArgs.NotifyHook = this;
+	
+	DetailsView = PropertyModule.CreateDetailView(DetailsArgs);
+
+	DetailsView->RegisterInstancedCustomPropertyLayout(UMediaProfile::StaticClass(), FOnGetDetailCustomizationInstance::CreateLambda([this]
+	{
+		return MakeShared<FMediaProfileTimecodeGenlockDetailsCustomization>();
+	}));
+
+	if (MediaProfile.IsValid())
+	{
+		DetailsView->SetObject(MediaProfile.Get());
+	}
+	
+	ChildSlot
+	[
+		SNew(SVerticalBox)
+		+SVerticalBox::Slot()
+		.AutoHeight()
+		[
+			SNew(SMediaFrameworkTimecodeGenlockHeader)
+		]
+
+		+SVerticalBox::Slot()
+		.AutoHeight()
+		[
+			SNew(SSeparator)
+			.Orientation(EOrientation::Orient_Horizontal)
+			.Thickness(2)
+		]
+		
+		+SVerticalBox::Slot()
+		.AutoHeight()
+		[
+			DetailsView->GetFilterAreaWidget().ToSharedRef()
+		]
+		
+		+SVerticalBox::Slot()
+		[
+			DetailsView.ToSharedRef()
+		]
+	];
+}
+
+void SMediaFrameworkTimecodeGenlockPanel::NotifyPostChange(const FPropertyChangedEvent& PropertyChangedEvent, FEditPropertyChain* PropertyThatChanged)
+{
+	TStrongObjectPtr<UMediaProfile> PinnedMediaProfile = MediaProfile.Pin();
+	if (!PinnedMediaProfile.IsValid())
+	{
+		return;
+	}
+
+	FProperty* HeadProperty = PropertyThatChanged->GetHead()->GetValue();
+	if (!HeadProperty)
+	{
+		return;
+	}
+	
+	if (HeadProperty->GetFName() == TEXT("TimecodeProvider") || HeadProperty->GetFName() == TEXT("bOverrideTimecodeProvider"))
+	{
+		PinnedMediaProfile->ApplyTimecodeProvider();
+	}
+
+	if (HeadProperty->GetFName() == TEXT("CustomTimeStep") || HeadProperty->GetFName() == TEXT("bOverrideCustomTimeStep"))
+	{
+		PinnedMediaProfile->ApplyCustomTimeStep();
+	}
+}
+
+#undef LOCTEXT_NAMESPACE
diff --git a/Engine/Plugins/Media/MediaFrameworkUtilities/Source/MediaFrameworkUtilitiesEditor/Private/UI/SMediaFrameworkTimecodeGenlockPanel.h b/Engine/Plugins/Media/MediaFrameworkUtilities/Source/MediaFrameworkUtilitiesEditor/Private/UI/SMediaFrameworkTimecodeGenlockPanel.h
new file mode 100644
--- /dev/null
+++ b/Engine/Plugins/Media/MediaFrameworkUtilities/Source/MediaFrameworkUtilitiesEditor/Private/UI/SMediaFrameworkTimecodeGenlockPanel.h
@@ -0,0 +1,35 @@
+﻿// Copyright Epic Games, Inc. All Rights Reserved.
+
+#pragma once
+
+#include "Misc/NotifyHook.h"
+#include "Widgets/SCompoundWidget.h"
+
+class IDetailsView;
+class UMediaProfile;
+
+/**
+ * Details panel that displays a media profile's Timecode and Genlock settings, as well as
+ * a timecode/genlock status bar and optionally a media profile selection dropdown
+ */
+class SMediaFrameworkTimecodeGenlockPanel : public SCompoundWidget, public FNotifyHook
+{
+public:
+	SLATE_BEGIN_ARGS(SMediaFrameworkTimecodeGenlockPanel) { }
+		SLATE_ARGUMENT(TWeakObjectPtr<UMediaProfile>, MediaProfile)
+	SLATE_END_ARGS()
+	
+	void Construct(const FArguments& InArgs);
+
+	//~ Begin FNotifyHook Interface
+	virtual void NotifyPostChange(const FPropertyChangedEvent& PropertyChangedEvent, FEditPropertyChain* PropertyThatChanged) override;
+	//~ End FNotifyHook Interface
+
+	
+private:
+	/** Media profile whose timecode/genlock settings are being displayed */
+	TWeakObjectPtr<UMediaProfile> MediaProfile;
+
+	/** Details view displaying timecode/genlock properties */
+	TSharedPtr<IDetailsView> DetailsView;
+};
