Media Profile: Specified category display order for AJA, Blackmagic, and NDI media sources and outputs so that categories are displayed in logical order within a details panel

diff --git a/Engine/Plugins/Media/AjaMedia/Source/AjaMedia/Public/AjaMediaSource.h b/Engine/Plugins/Media/AjaMedia/Source/AjaMedia/Public/AjaMediaSource.h
--- a/Engine/Plugins/Media/AjaMedia/Source/AjaMedia/Public/AjaMediaSource.h
+++ b/Engine/Plugins/Media/AjaMedia/Source/AjaMedia/Public/AjaMediaSource.h
@@ -32,7 +32,7 @@ enum class EAjaMediaAudioChannel : uint8
 /**
  * Media source for AJA streams.
  */
-UCLASS(BlueprintType, hideCategories=(Platforms,Object), meta=(MediaIOCustomLayout="AJA"))
+UCLASS(BlueprintType, hideCategories=(Platforms,Object), PrioritizeCategories=(AJA, Video, Audio, Ancillary, Synchronization, Debug), meta=(MediaIOCustomLayout="AJA"))
 class AJAMEDIA_API UAjaMediaSource : public UCaptureCardMediaSource
 {
 	GENERATED_BODY()
diff --git a/Engine/Plugins/Media/AjaMedia/Source/AjaMediaOutput/Public/AjaMediaOutput.h b/Engine/Plugins/Media/AjaMedia/Source/AjaMediaOutput/Public/AjaMediaOutput.h
--- a/Engine/Plugins/Media/AjaMedia/Source/AjaMediaOutput/Public/AjaMediaOutput.h
+++ b/Engine/Plugins/Media/AjaMedia/Source/AjaMediaOutput/Public/AjaMediaOutput.h
@@ -41,7 +41,7 @@ enum class EAjaMediaOutputChannelConfiguration : uint8
  * @note	'Frame Buffer Pixel Format' must be set to at least 8 bits of alpha to enabled the Key.
  * @note	'Enable alpha channel support in post-processing' must be set to 'Allow through tonemapper' to enabled the Key.
  */
-UCLASS(BlueprintType, meta=(MediaIOCustomLayout="AJA"))
+UCLASS(BlueprintType, PrioritizeCategories=(AJA, Output, Synchronization, Debug), meta=(MediaIOCustomLayout="AJA"))
 class AJAMEDIAOUTPUT_API UAjaMediaOutput : public UMediaOutput
 {
 	GENERATED_UCLASS_BODY()
diff --git a/Engine/Plugins/Media/BlackmagicMedia/Source/BlackmagicMedia/Public/BlackmagicMediaSource.h b/Engine/Plugins/Media/BlackmagicMedia/Source/BlackmagicMedia/Public/BlackmagicMediaSource.h
--- a/Engine/Plugins/Media/BlackmagicMedia/Source/BlackmagicMedia/Public/BlackmagicMediaSource.h
+++ b/Engine/Plugins/Media/BlackmagicMedia/Source/BlackmagicMedia/Public/BlackmagicMediaSource.h
@@ -33,7 +33,7 @@ enum class EBlackmagicMediaAudioChannel : uint8
 /**
  * Media source description for Blackmagic.
  */
-UCLASS(BlueprintType, hideCategories=(Platforms,Object), meta=(MediaIOCustomLayout="Blackmagic"))
+UCLASS(BlueprintType, hideCategories=(Platforms,Object), PrioritizeCategories=(Blackmagic, Video, Audio, Ancillary, Synchronization, Debug), meta=(MediaIOCustomLayout="Blackmagic"))
 class BLACKMAGICMEDIA_API UBlackmagicMediaSource : public UCaptureCardMediaSource
 {
 	GENERATED_BODY()
diff --git a/Engine/Plugins/Media/BlackmagicMedia/Source/BlackmagicMediaOutput/Public/BlackmagicMediaOutput.h b/Engine/Plugins/Media/BlackmagicMedia/Source/BlackmagicMediaOutput/Public/BlackmagicMediaOutput.h
--- a/Engine/Plugins/Media/BlackmagicMedia/Source/BlackmagicMediaOutput/Public/BlackmagicMediaOutput.h
+++ b/Engine/Plugins/Media/BlackmagicMedia/Source/BlackmagicMediaOutput/Public/BlackmagicMediaOutput.h
@@ -45,7 +45,7 @@ enum class EBlackmagicMediaOutputAudioBitDepth : uint8
  * @note	'Frame Buffer Pixel Format' must be set to at least 8 bits of alpha to enabled the Key.
  * @note	'Enable alpha channel support in post-processing' must be set to 'Allow through tonemapper' to enabled the Key.
  */
-UCLASS(BlueprintType, meta = (MediaIOCustomLayout = "Blackmagic"))
+UCLASS(BlueprintType, PrioritizeCategories=(Blackmagic, Output, Synchronization, Debug), meta = (MediaIOCustomLayout = "Blackmagic"))
 class BLACKMAGICMEDIAOUTPUT_API UBlackmagicMediaOutput : public UMediaOutput
 {
 	GENERATED_UCLASS_BODY()
diff --git a/Engine/Plugins/Media/MediaFrameworkUtilities/Source/MediaFrameworkUtilitiesEditor/Private/AssetEditor/MediaProfileMediaItemDetailCustomization.cpp b/Engine/Plugins/Media/MediaFrameworkUtilities/Source/MediaFrameworkUtilitiesEditor/Private/AssetEditor/MediaProfileMediaItemDetailCustomization.cpp
--- a/Engine/Plugins/Media/MediaFrameworkUtilities/Source/MediaFrameworkUtilitiesEditor/Private/AssetEditor/MediaProfileMediaItemDetailCustomization.cpp
+++ b/Engine/Plugins/Media/MediaFrameworkUtilities/Source/MediaFrameworkUtilitiesEditor/Private/AssetEditor/MediaProfileMediaItemDetailCustomization.cpp
@@ -324,6 +324,11 @@ void FMediaProfileMediaOutputDetailCustomization::CustomizeDetails(IDetailLayout
 {
 	FMediaProfileMediaItemDetailCustomization::CustomizeDetails(DetailBuilder);
 
+	// Hack to avoid the details builder separating out categories that only have advanced properties in them and putting them at the bottom.
+	// If there is any sort function provided, the builder lumps all categories together and sorts via their sort order. We don't need to manually
+	// adjust the sort order, but we do need the details builder to lump all categories together, hence the empty sort function
+	DetailBuilder.SortCategories([](const TMap<FName, IDetailCategoryBuilder*>& CategoryMap) { });
+	
 	MediaProfileEditor.Pin()->GetOnCaptureSettingsObjectChanged().AddSP(this, &FMediaProfileMediaOutputDetailCustomization::OnCaptureSettingsObjectChanged);
 	
 	TArray<TWeakObjectPtr<UObject>> Objects;
diff --git a/Engine/Plugins/Media/NDIMedia/Source/NDIMedia/Public/NDIMediaOutput.h b/Engine/Plugins/Media/NDIMedia/Source/NDIMedia/Public/NDIMediaOutput.h
--- a/Engine/Plugins/Media/NDIMedia/Source/NDIMedia/Public/NDIMediaOutput.h
+++ b/Engine/Plugins/Media/NDIMedia/Source/NDIMedia/Public/NDIMediaOutput.h
@@ -49,7 +49,7 @@ enum class ENDIMediaOutputAudioSampleRate : uint32
  * Output information for a NDI media capture.
  * @note	'Frame Buffer Pixel Format' must be set to RGBA8
  */
-UCLASS(BlueprintType, meta = (MediaIOCustomLayout = "NDI"))
+UCLASS(BlueprintType, PrioritizeCategories=(Media, Audio, Output, Synchronization), meta = (MediaIOCustomLayout = "NDI"))
 class NDIMEDIA_API UNDIMediaOutput : public UMediaOutput
 {
 	GENERATED_BODY()
diff --git a/Engine/Plugins/Media/NDIMedia/Source/NDIMedia/Public/NDIMediaSource.h b/Engine/Plugins/Media/NDIMedia/Source/NDIMedia/Public/NDIMediaSource.h
--- a/Engine/Plugins/Media/NDIMedia/Source/NDIMedia/Public/NDIMediaSource.h
+++ b/Engine/Plugins/Media/NDIMedia/Source/NDIMedia/Public/NDIMediaSource.h
@@ -12,7 +12,7 @@
 /**
  * Media source for NDI streams.
  */
-UCLASS(MinimalAPI, BlueprintType, hideCategories=(Platforms,Object), meta=(MediaIOCustomLayout="NDI"))
+UCLASS(MinimalAPI, BlueprintType, hideCategories=(Platforms,Object), PrioritizeCategories=(NDI, Video, Audio, Ancillary, Synchronization, Debug), meta=(MediaIOCustomLayout="NDI"))
 class UNDIMediaSource : public UCaptureCardMediaSource
 {
 	GENERATED_BODY()
