Media Profile: Added section pills to media profile editor details panel

diff --git a/Engine/Plugins/Media/MediaFrameworkUtilities/Source/MediaFrameworkUtilitiesEditor/Private/AssetEditor/MediaProfileMediaItemDetailCustomization.cpp b/Engine/Plugins/Media/MediaFrameworkUtilities/Source/MediaFrameworkUtilitiesEditor/Private/AssetEditor/MediaProfileMediaItemDetailCustomization.cpp
--- a/Engine/Plugins/Media/MediaFrameworkUtilities/Source/MediaFrameworkUtilitiesEditor/Private/AssetEditor/MediaProfileMediaItemDetailCustomization.cpp
+++ b/Engine/Plugins/Media/MediaFrameworkUtilities/Source/MediaFrameworkUtilitiesEditor/Private/AssetEditor/MediaProfileMediaItemDetailCustomization.cpp
@@ -16,6 +16,21 @@ FText FMediaProfileMediaSourceDetailCustomization::GetMediaItemTypeText() const
 	return LOCTEXT("MediaSourceTypeText", "Media Source");
 }
 
+void FMediaProfileMediaSourceDetailCustomization::RegisterMediaTypeSection(FPropertyEditorModule& PropertyModule, const FName& MediaTypeCategory)
+{
+	static bool bRegistered = false;
+	if (!bRegistered)
+	{
+		const TSharedPtr<FPropertySection> Section = PropertyModule.FindOrCreateSection(
+			UMediaSource::StaticClass()->GetFName(),
+			TEXT("MediaSourceSection"),
+			GetMediaItemTypeText());
+
+		Section->AddCategory(MediaTypeCategory);
+		bRegistered = true;
+	}
+}
+
 FText FMediaProfileMediaSourceDetailCustomization::GetMediaObjectLabel(UMediaSource* InMediaObject) const
 {
 	if (!MediaProfile.IsValid())
@@ -323,7 +338,7 @@ FMediaProfileMediaOutputDetailCustomization::~FMediaProfileMediaOutputDetailCust
 void FMediaProfileMediaOutputDetailCustomization::CustomizeDetails(IDetailLayoutBuilder& DetailBuilder)
 {
 	FMediaProfileMediaItemDetailCustomization::CustomizeDetails(DetailBuilder);
-
+	
 	// Hack to avoid the details builder separating out categories that only have advanced properties in them and putting them at the bottom.
 	// If there is any sort function provided, the builder lumps all categories together and sorts via their sort order. We don't need to manually
 	// adjust the sort order, but we do need the details builder to lump all categories together, hence the empty sort function
@@ -349,11 +364,14 @@ void FMediaProfileMediaOutputDetailCustomization::CustomizeDetails(IDetailLayout
 		}));
 	
 	// For media outputs, add the corresponding capture info to the details panel
+	const FName CaptureCategoryName = TEXT("MediaCaptureCategory");
 	IDetailCategoryBuilder& MediaCaptureCategory = DetailBuilder.EditCategory(
-		TEXT("MediaCaptureCategory"),
+		CaptureCategoryName,
 		LOCTEXT("MediaCaptureCategoryLabel", "Capture"),
 		ECategoryPriority::Important);
 
+	RegisterCaptureSection(FModuleManager::GetModuleChecked<FPropertyEditorModule>("PropertyEditor"), CaptureCategoryName);
+
 	MediaCaptureCategory
 		.AddCustomRow(LOCTEXT("MediaCaptureButtonFilter", "Start/Stop Media Capture"))
 		.Visibility(TAttribute<EVisibility>::CreateSP(this, &FMediaProfileMediaOutputDetailCustomization::GetValidObjectVisibility))
@@ -471,6 +489,21 @@ FText FMediaProfileMediaOutputDetailCustomization::GetMediaItemTypeText() const
 	return LOCTEXT("MediaOutputTypeText", "Media Output");
 }
 
+void FMediaProfileMediaOutputDetailCustomization::RegisterMediaTypeSection(FPropertyEditorModule& PropertyModule, const FName& MediaTypeCategory)
+{
+	static bool bRegistered = false;
+	if (!bRegistered)
+	{
+		const TSharedPtr<FPropertySection> Section = PropertyModule.FindOrCreateSection(
+			UMediaOutput::StaticClass()->GetFName(),
+			TEXT("MediaOutputSection"),
+			GetMediaItemTypeText());
+
+		Section->AddCategory(MediaTypeCategory);
+		bRegistered = true;
+	}
+}
+
 FText FMediaProfileMediaOutputDetailCustomization::GetMediaObjectLabel(UMediaOutput* InMediaObject) const
 {
 	if (!MediaProfile.IsValid())
@@ -555,6 +588,21 @@ void FMediaProfileMediaOutputDetailCustomization::SetMediaObject(UMediaOutput* I
 	}
 }
 
+void FMediaProfileMediaOutputDetailCustomization::RegisterCaptureSection(FPropertyEditorModule& PropertyModule, const FName& CaptureCategory)
+{
+	static bool bRegistered = false;
+	if (!bRegistered)
+	{
+		const TSharedPtr<FPropertySection> Section = PropertyModule.FindOrCreateSection(
+			UMediaOutput::StaticClass()->GetFName(),
+			TEXT("CaptureSection"),
+			LOCTEXT("CaptureSection", "Capture"));
+
+		Section->AddCategory(CaptureCategory);
+		bRegistered = true;
+	}
+}
+
 void FMediaProfileMediaOutputDetailCustomization::OnCaptureButtonPressed()
 {
 	if (MediaProfileEditor.IsValid())
diff --git a/Engine/Plugins/Media/MediaFrameworkUtilities/Source/MediaFrameworkUtilitiesEditor/Private/AssetEditor/MediaProfileMediaItemDetailCustomization.h b/Engine/Plugins/Media/MediaFrameworkUtilities/Source/MediaFrameworkUtilitiesEditor/Private/AssetEditor/MediaProfileMediaItemDetailCustomization.h
--- a/Engine/Plugins/Media/MediaFrameworkUtilities/Source/MediaFrameworkUtilitiesEditor/Private/AssetEditor/MediaProfileMediaItemDetailCustomization.h
+++ b/Engine/Plugins/Media/MediaFrameworkUtilities/Source/MediaFrameworkUtilitiesEditor/Private/AssetEditor/MediaProfileMediaItemDetailCustomization.h
@@ -9,6 +9,7 @@
 #include "DetailLayoutBuilder.h"
 #include "DetailWidgetRow.h"
 #include "DummyMediaObject.h"
+#include "EditorCategoryUtils.h"
 #include "IDetailCustomization.h"
 #include "MediaProfileEditor.h"
 #include "ScopedTransaction.h"
@@ -67,9 +68,17 @@ public:
 
 	virtual void CustomizeDetails(IDetailLayoutBuilder& DetailBuilder) override
 	{
+		FPropertyEditorModule& PropertyModule = FModuleManager::GetModuleChecked<FPropertyEditorModule>("PropertyEditor");
+		
+		// Register any unregistered categories as sections to give it a section pill. Need to do this every time as this details panel can show
+		// any number of different subtypes of UMediaSource or UMediaOutput, which may have its own categories
+		AddSectionsForEachCategory(DetailBuilder, PropertyModule);
+		
 		const FText CategoryLabel = FText::Format(LOCTEXT("MediaItemCategoryName", "{0}"), GetMediaItemTypeText());
-		IDetailCategoryBuilder& MediaTypeCategory = DetailBuilder.EditCategory(TEXT("MediaTypeCategory"), CategoryLabel, ECategoryPriority::Important);
-
+		const FName MediaTypeCategoryName = TEXT("MediaTypeCategory");
+		IDetailCategoryBuilder& MediaTypeCategory = DetailBuilder.EditCategory(MediaTypeCategoryName, CategoryLabel, ECategoryPriority::Important);
+		RegisterMediaTypeSection(PropertyModule, MediaTypeCategoryName);
+		
 		MediaTypeCategory.AddCustomRow(LOCTEXT("LabelFilterText", "Label"))
 			.Visibility(TAttribute<EVisibility>::CreateSP(this, &FMediaProfileMediaItemDetailCustomization::GetValidObjectVisibility))
 			.NameContent()
@@ -132,7 +141,31 @@ public:
 	}
 
 protected:
+	/** Registers sections to the root media item type for each category being displayed in the details panel if one does not yet exist */
+	void AddSectionsForEachCategory(IDetailLayoutBuilder& DetailBuilder, FPropertyEditorModule& PropertyModule)
+	{		
+		TArray<FName> Categories;
+		DetailBuilder.GetCategoryNames(Categories);
+
+		const FName TypeName = TMediaType::StaticClass()->GetFName();
+		for (const FName& CategoryName : Categories)
+		{
+			const TSharedPtr<FPropertySection> Section = PropertyModule.FindOrCreateSection(
+				TypeName,
+				CategoryName,
+				FEditorCategoryUtils::GetCategoryDisplayString(FText::FromName(CategoryName)));
+
+			if (!Section->HasAddedCategory(CategoryName))
+			{
+				Section->AddCategory(CategoryName);
+			}
+		}
+	}
+	
 	virtual FText GetMediaItemTypeText() const = 0;
+
+	/** Registers a section for the custom media type category under the type specific name */
+	virtual void RegisterMediaTypeSection(FPropertyEditorModule& PropertyModule, const FName& MediaTypeCategory) = 0;
 	
 	EVisibility GetValidObjectVisibility() const
 	{
@@ -305,6 +338,7 @@ public:
 
 private:
 	virtual FText GetMediaItemTypeText() const override;
+	virtual void RegisterMediaTypeSection(FPropertyEditorModule& PropertyModule, const FName& MediaTypeCategory) override;
 	virtual FText GetMediaObjectLabel(UMediaSource* InMediaObject) const override;
 	virtual void SetMediaObjectLabel(UMediaSource* InMediaObject, const FText& InLabel) override;
 	virtual void SetMediaObject(UMediaSource* InOriginalMediaObject, UMediaSource* InNewMediaObject) override;
@@ -326,9 +360,13 @@ public:
 	
 private:
 	virtual FText GetMediaItemTypeText() const override;
+	virtual void RegisterMediaTypeSection(FPropertyEditorModule& PropertyModule, const FName& MediaTypeCategory) override;
 	virtual FText GetMediaObjectLabel(UMediaOutput* InMediaObject) const override;
 	virtual void SetMediaObjectLabel(UMediaOutput* InMediaObject, const FText& InLabel) override;
 	virtual void SetMediaObject(UMediaOutput* InOriginalMediaObject, UMediaOutput* InNewMediaObject) override;
+
+	/** Registers a section for the Capture category */
+	void RegisterCaptureSection(FPropertyEditorModule& PropertyModule, const FName& CaptureCategory);
 	
 	void OnCaptureButtonPressed();
 	
