Media Profile: Fixed issue where deleting open media sources in the media profile editor could result in a media player attempting to open an SDI connection that was already open, preventing it from displaying properly. Now, the media profile playback manager will attempt to move existing media players to the correct slots and textures when a source is deleted or moved. Also set AutoClear on the default media textures to ensure they clear themselves when a media player is closed.

#virtualized

diff --git a/Engine/Plugins/Media/MediaFrameworkUtilities/Source/MediaFrameworkUtilities/Private/Profile/MediaProfile.cpp b/Engine/Plugins/Media/MediaFrameworkUtilities/Source/MediaFrameworkUtilities/Private/Profile/MediaProfile.cpp
--- a/Engine/Plugins/Media/MediaFrameworkUtilities/Source/MediaFrameworkUtilities/Private/Profile/MediaProfile.cpp
+++ b/Engine/Plugins/Media/MediaFrameworkUtilities/Source/MediaFrameworkUtilities/Private/Profile/MediaProfile.cpp
@@ -117,6 +117,7 @@ bool UMediaProfile::RemoveMediaSource(int32 Index)
 		return false;
 	}
 	
+	PlaybackManager->ChangeMediaSourceIndex(Index, INDEX_NONE);
 	MediaSources.RemoveAt(Index, EAllowShrinking::No);
 
 #if WITH_EDITOR
@@ -136,6 +137,7 @@ bool UMediaProfile::MoveMediaSource(int32 CurrentIndex, int32 DestIndex)
 		return false;
 	}
 
+	PlaybackManager->ChangeMediaSourceIndex(CurrentIndex, DestIndex);
 	TObjectPtr<UMediaSource> MediaSource = MediaSources[CurrentIndex];
 	MediaSources.RemoveAt(CurrentIndex, EAllowShrinking::No);
 	MediaSources.Insert(MediaSource, DestIndex);
diff --git a/Engine/Plugins/Media/MediaFrameworkUtilities/Source/MediaFrameworkUtilities/Private/Profile/MediaProfilePlaybackManager.cpp b/Engine/Plugins/Media/MediaFrameworkUtilities/Source/MediaFrameworkUtilities/Private/Profile/MediaProfilePlaybackManager.cpp
--- a/Engine/Plugins/Media/MediaFrameworkUtilities/Source/MediaFrameworkUtilities/Private/Profile/MediaProfilePlaybackManager.cpp
+++ b/Engine/Plugins/Media/MediaFrameworkUtilities/Source/MediaFrameworkUtilities/Private/Profile/MediaProfilePlaybackManager.cpp
@@ -83,15 +83,24 @@ UMediaTexture* UMediaProfilePlaybackManager::OpenSourceFromIndex(int32 InMediaSo
 	{
 		return nullptr;
 	}
-
+		
+	UMediaTexture* MediaTexture = GetSourceMediaTextureFromIndex(InMediaSourceIndex);
+	if (!MediaTexture)
+	{
+		return nullptr;
+	}
+	
 	UMediaSource* MediaSource = MediaProfile->GetMediaSource(InMediaSourceIndex);
-	UMediaPlayer* MediaSourcePlayer = FindOrCreateMediaPlayerForIndex(InMediaSourceIndex);	
+	UMediaPlayer* MediaSourcePlayer = FindOrCreateMediaPlayerForIndex(InMediaSourceIndex);
+
+	MediaTexture->SetMediaPlayer(MediaSourcePlayer);
+	
 	if (MediaSourcePlayer->IsPlaying())
 	{
 		// If the media source is already being played, simply return the media texture
 		if (MediaSourcePlayer->GetUrl() == MediaSource->GetUrl())
 		{
-			return MediaSourceTextures[InMediaSourceIndex].LoadSynchronous();
+			return MediaTexture;
 		}
 	}
 
@@ -101,7 +110,7 @@ UMediaTexture* UMediaProfilePlaybackManager::OpenSourceFromIndex(int32 InMediaSo
 	
 	if (MediaSourcePlayer->OpenSourceWithOptions(MediaSource, Options))
 	{
-		return MediaSourceTextures[InMediaSourceIndex].LoadSynchronous();
+		return MediaTexture;
 	}
 	
 	// If the media player was unable to open the source, return nullptr to indicate that the media texture will not be displaying the media source
@@ -128,7 +137,11 @@ void UMediaProfilePlaybackManager::CloseSourceFromIndex(int32 InMediaSourceIndex
 	}
 
 	MediaSourcePlayers[InMediaSourceIndex]->Close();
-
+	if (UMediaTexture* MediaTexture = GetSourceMediaTextureFromIndex(InMediaSourceIndex))
+	{
+		MediaTexture->SetMediaPlayer(nullptr);
+	}
+	
 	if (bDestroyMediaPlayer)
 	{
 		MediaSourcePlayers.Remove(InMediaSourceIndex);
@@ -163,17 +176,80 @@ UMediaPlayer* UMediaProfilePlaybackManager::FindOrCreateMediaPlayerForIndex(int3
 	if (!MediaSourcePlayers.Contains(InMediaSourceIndex))
 	{
 		UMediaPlayer* NewMediaPlayer = NewObject<UMediaPlayer>(this);
-		if (UMediaTexture* MediaTexture = GetSourceMediaTextureFromIndex(InMediaSourceIndex))
-		{
-			MediaTexture->SetMediaPlayer(NewMediaPlayer);
-		}
-		
 		MediaSourcePlayers.Add(InMediaSourceIndex, NewMediaPlayer);
 	}
 
 	return MediaSourcePlayers[InMediaSourceIndex];
 }
 
+void UMediaProfilePlaybackManager::ChangeMediaSourceIndex(int32 InOriginalMediaSourceIndex, int32 InNewMediaSourceIndex)
+{
+	UMediaProfile* MediaProfile = GetOwningMediaProfile();
+	if (!MediaProfile)
+	{
+		return;
+	}
+
+	if (InOriginalMediaSourceIndex == InNewMediaSourceIndex)
+	{
+		return;
+	}
+	
+	// First, remove the existing media player from the list of media players and shift all slots below it up one
+	UMediaPlayer* MediaPlayer = nullptr;
+	if (InOriginalMediaSourceIndex != INDEX_NONE)
+	{
+		MediaPlayer = MediaSourcePlayers.Contains(InOriginalMediaSourceIndex) ? MediaSourcePlayers[InOriginalMediaSourceIndex] : nullptr;
+		int32 LastIndex = INDEX_NONE;
+		for (int32 Index = InOriginalMediaSourceIndex + 1; Index < MediaProfile->NumMediaSources(); ++Index)
+		{
+			if (MediaSourcePlayers.Contains(Index))
+			{
+				MediaSourcePlayers.Add(Index - 1, MediaSourcePlayers[Index]);
+				LastIndex = Index;
+			}
+		}
+
+		if (LastIndex != INDEX_NONE)
+		{
+			MediaSourcePlayers.Remove(LastIndex);
+		}
+	}
+
+	// Second, insert the media player back into the list, shifting all slots at or below that index down one
+	if (InNewMediaSourceIndex != INDEX_NONE)
+	{
+		for (int32 Index = MediaProfile->NumMediaSources() - 1; Index >= InNewMediaSourceIndex; --Index)
+		{
+			if (MediaSourcePlayers.Contains(Index))
+			{
+				MediaSourcePlayers.Add(Index + 1, MediaSourcePlayers[Index]);
+			}
+		}
+
+		if (MediaPlayer)
+		{
+			MediaSourcePlayers.Add(InNewMediaSourceIndex, MediaPlayer);
+		}
+	}
+
+	// Finally, refresh all media textures to have the correct media player attached
+	for (int32 Index = 0; Index < MediaSourceTextures.Num(); ++Index)
+	{
+		if (UMediaTexture* MediaTexture = GetSourceMediaTextureFromIndex(Index))
+		{
+			if (MediaSourcePlayers.Contains(Index))
+			{
+				MediaTexture->SetMediaPlayer(MediaSourcePlayers[Index]);
+			}
+			else
+			{
+				MediaTexture->SetMediaPlayer(nullptr);
+			}
+		}
+	}
+}
+
 UMediaProfile* UMediaProfilePlaybackManager::GetOwningMediaProfile() const
 {
 	return GetTypedOuter<UMediaProfile>();
diff --git a/Engine/Plugins/Media/MediaFrameworkUtilities/Source/MediaFrameworkUtilities/Public/Profile/MediaProfilePlaybackManager.h b/Engine/Plugins/Media/MediaFrameworkUtilities/Source/MediaFrameworkUtilities/Public/Profile/MediaProfilePlaybackManager.h
--- a/Engine/Plugins/Media/MediaFrameworkUtilities/Source/MediaFrameworkUtilities/Public/Profile/MediaProfilePlaybackManager.h
+++ b/Engine/Plugins/Media/MediaFrameworkUtilities/Source/MediaFrameworkUtilities/Public/Profile/MediaProfilePlaybackManager.h
@@ -71,6 +71,13 @@ private:
 	 */
 	UMediaPlayer* FindOrCreateMediaPlayerForIndex(int32 InMediaSourceIndex);
 
+	/**
+	 * Processes a change to a media source's slot index, moving any existing media players around to the correct media textures
+	 * @param InOriginalMediaSourceIndex The original index of the media source, or INDEX_NONE if we are processing an insert
+	 * @param InNewMediaSourceIndex The new index of the media source, or INDEX_NONE if we are processing a delete
+	 */
+	void ChangeMediaSourceIndex(int32 InOriginalMediaSourceIndex, int32 InNewMediaSourceIndex);
+	
 	/** Gets the media profile that owns this playback manager */
 	UMediaProfile* GetOwningMediaProfile() const;
 	
@@ -82,4 +89,6 @@ private:
 	/** A list of media textures to use to output the corresponding media sources */
 	UPROPERTY(Transient)
 	TArray<TSoftObjectPtr<UMediaTexture>> MediaSourceTextures;
+
+	friend class UMediaProfile;
 };
diff --git a/Engine/Plugins/Media/MediaFrameworkUtilities/Source/MediaFrameworkUtilitiesEditor/Private/AssetEditor/MediaProfileEditor.cpp b/Engine/Plugins/Media/MediaFrameworkUtilities/Source/MediaFrameworkUtilitiesEditor/Private/AssetEditor/MediaProfileEditor.cpp
--- a/Engine/Plugins/Media/MediaFrameworkUtilities/Source/MediaFrameworkUtilitiesEditor/Private/AssetEditor/MediaProfileEditor.cpp
+++ b/Engine/Plugins/Media/MediaFrameworkUtilities/Source/MediaFrameworkUtilitiesEditor/Private/AssetEditor/MediaProfileEditor.cpp
@@ -237,7 +237,8 @@ FMediaProfileEditor::~FMediaProfileEditor()
 	FEditorDelegates::MapChange.RemoveAll(this);
 	FEditorDelegates::PrePIEEnded.RemoveAll(this);
 	FEditorDelegates::PostPIEStarted.RemoveAll(this);
-	
+
+	CloseAllMediaSources();
 	StopCapturingAllMediaOutputs();
 }
 
@@ -306,6 +307,14 @@ FLinearColor FMediaProfileEditor::GetWorldCentricTabColorScale() const
 	return FLinearColor::White;
 }
 
+void FMediaProfileEditor::CloseAllMediaSources()
+{
+	for (int32 Index = 0; Index < MediaProfileBeingEdited->NumMediaSources(); ++Index)
+	{
+		MediaProfileBeingEdited->GetPlaybackManager()->CloseSourceFromIndex(Index);
+	}
+}
+
 TSharedPtr<FMediaProfileEditor::FViewportCapture> FMediaProfileEditor::FindOrCreateViewportCapture(UMediaOutput* InMediaOutput)
 {
 	UMediaFrameworkWorldSettingsAssetUserData* CaptureSettings = FindOrAddMediaFrameworkCaptureSettings();
diff --git a/Engine/Plugins/Media/MediaFrameworkUtilities/Source/MediaFrameworkUtilitiesEditor/Private/AssetEditor/MediaProfileEditor.h b/Engine/Plugins/Media/MediaFrameworkUtilities/Source/MediaFrameworkUtilitiesEditor/Private/AssetEditor/MediaProfileEditor.h
--- a/Engine/Plugins/Media/MediaFrameworkUtilities/Source/MediaFrameworkUtilitiesEditor/Private/AssetEditor/MediaProfileEditor.h
+++ b/Engine/Plugins/Media/MediaFrameworkUtilities/Source/MediaFrameworkUtilitiesEditor/Private/AssetEditor/MediaProfileEditor.h
@@ -63,6 +63,9 @@ public:
 
 	/** Gets the media profile being edited */
 	UMediaProfile* GetMediaProfile() const { return MediaProfileBeingEdited; }
+
+	/** Closes all open media sources in the media profile */
+	void CloseAllMediaSources();
 	
 	/**
 	 * Attempts to find an existing viewport capture for the specified media output, or creates one if one does not exist
