Media Profile: Set the media profile editor 2.0 to be default editor and fixed a variety of bugs:
Fixed issues with media capture when changing media output type.
Fixed issue where added media source or output is not auto-selected.
Fixed issue where opening multiple media profiles can cause media sources and outputs to be removed from them.

diff --git a/Engine/Plugins/Media/MediaFrameworkUtilities/Source/MediaFrameworkUtilities/Private/Profile/MediaProfile.cpp b/Engine/Plugins/Media/MediaFrameworkUtilities/Source/MediaFrameworkUtilities/Private/Profile/MediaProfile.cpp
--- a/Engine/Plugins/Media/MediaFrameworkUtilities/Source/MediaFrameworkUtilities/Private/Profile/MediaProfile.cpp
+++ b/Engine/Plugins/Media/MediaFrameworkUtilities/Source/MediaFrameworkUtilities/Private/Profile/MediaProfile.cpp
@@ -105,7 +105,15 @@ void UMediaProfile::AddMediaSource(UMediaSource* MediaSource)
 	else
 	{
 		FString Name = MediaSource ? MediaSource->GetName() : TEXT("");
-		MediaSourceLabels.Add(Name);
+		const int32 Index = MediaSources.Num() - 1;
+		if (MediaSourceLabels.IsValidIndex(Index))
+		{
+			MediaSourceLabels[Index] = Name;
+		}
+		else
+		{
+			MediaSourceLabels.Add(Name);
+		}
 	}
 #endif
 }
@@ -264,7 +272,15 @@ void UMediaProfile::AddMediaOutput(UMediaOutput* MediaOutput)
 	else
 	{
 		FString Name = MediaOutput ? MediaOutput->GetName() : TEXT("");
-		MediaOutputLabels.Add(Name);
+		const int32 Index = MediaOutputs.Num() - 1;
+		if (MediaOutputLabels.IsValidIndex(Index))
+		{
+			MediaOutputLabels[Index] = Name;
+		}
+		else
+		{
+			MediaOutputLabels.Add(Name);
+		}
 	}
 #endif
 }
diff --git a/Engine/Plugins/Media/MediaFrameworkUtilities/Source/MediaFrameworkUtilitiesEditor/Private/AssetEditor/MediaProfileEditor.cpp b/Engine/Plugins/Media/MediaFrameworkUtilities/Source/MediaFrameworkUtilitiesEditor/Private/AssetEditor/MediaProfileEditor.cpp
--- a/Engine/Plugins/Media/MediaFrameworkUtilities/Source/MediaFrameworkUtilitiesEditor/Private/AssetEditor/MediaProfileEditor.cpp
+++ b/Engine/Plugins/Media/MediaFrameworkUtilities/Source/MediaFrameworkUtilitiesEditor/Private/AssetEditor/MediaProfileEditor.cpp
@@ -223,8 +223,8 @@ void FMediaProfileEditor::Initialize(const EToolkitMode::Type Mode, const TShare
 	
 	// TODO: Not sure this is the best place for this, but for now, when we open a media profile, automatically make it the active media profile
 	// First, update the media profile settings config to match the media profile's source and output count
-	GetMutableDefault<UMediaProfileSettings>()->FillDefaultMediaSourceProxies(InMediaProfile->NumMediaSources());
-	GetMutableDefault<UMediaProfileSettings>()->FillDefaultMediaOutputProxies(InMediaProfile->NumMediaOutputs());
+	GetMutableDefault<UMediaProfileSettings>()->FillDefaultMediaSourceProxies(InMediaProfile->NumMediaSources(), false);
+	GetMutableDefault<UMediaProfileSettings>()->FillDefaultMediaOutputProxies(InMediaProfile->NumMediaOutputs(), false);
 	IMediaProfileManager::Get().SetCurrentMediaProfile(InMediaProfile);
 }
 
diff --git a/Engine/Plugins/Media/MediaFrameworkUtilities/Source/MediaFrameworkUtilitiesEditor/Private/AssetEditor/MediaProfileEditorToolkit.cpp b/Engine/Plugins/Media/MediaFrameworkUtilities/Source/MediaFrameworkUtilitiesEditor/Private/AssetEditor/MediaProfileEditorToolkit.cpp
deleted file mode 100644
diff --git a/Engine/Plugins/Media/MediaFrameworkUtilities/Source/MediaFrameworkUtilitiesEditor/Private/AssetEditor/MediaProfileEditorToolkit.h b/Engine/Plugins/Media/MediaFrameworkUtilities/Source/MediaFrameworkUtilitiesEditor/Private/AssetEditor/MediaProfileEditorToolkit.h
deleted file mode 100644
diff --git a/Engine/Plugins/Media/MediaFrameworkUtilities/Source/MediaFrameworkUtilitiesEditor/Private/AssetEditor/MediaProfileMediaItemDetailCustomization.cpp b/Engine/Plugins/Media/MediaFrameworkUtilities/Source/MediaFrameworkUtilitiesEditor/Private/AssetEditor/MediaProfileMediaItemDetailCustomization.cpp
--- a/Engine/Plugins/Media/MediaFrameworkUtilities/Source/MediaFrameworkUtilitiesEditor/Private/AssetEditor/MediaProfileMediaItemDetailCustomization.cpp
+++ b/Engine/Plugins/Media/MediaFrameworkUtilities/Source/MediaFrameworkUtilitiesEditor/Private/AssetEditor/MediaProfileMediaItemDetailCustomization.cpp
@@ -538,18 +538,25 @@ void FMediaProfileMediaOutputDetailCustomization::SetMediaObjectLabel(UMediaOutp
 
 void FMediaProfileMediaOutputDetailCustomization::SetMediaObject(UMediaOutput* InOriginalMediaObject, UMediaOutput* InNewMediaObject)
 {
-	if (!MediaProfile.IsValid())
+	TSharedPtr<FMediaProfileEditor> PinnedMediaProfileEditor = MediaProfileEditor.Pin();
+	if (!PinnedMediaProfileEditor.IsValid())
+	{
+		return;
+	}
+
+	TStrongObjectPtr<UMediaProfile> PinnedMediaProfile = MediaProfile.Pin();
+	if (!PinnedMediaProfile.IsValid())
 	{
 		return;
 	}
 	
-	int32 Index = MediaProfile->FindMediaOutputIndex(InOriginalMediaObject);
+	int32 Index = PinnedMediaProfile->FindMediaOutputIndex(InOriginalMediaObject);
 	if (InOriginalMediaObject->IsA<UDummyMediaOutput>())
 	{
 		Index = Cast<UDummyMediaOutput>(InOriginalMediaObject)->MediaProfileIndex;
 	}
-		
-	MediaProfile->SetMediaOutput(Index, InNewMediaObject);
+
+	PinnedMediaProfileEditor->StopCapturingMediaOutput(InOriginalMediaObject);
 	
 	UMediaFrameworkWorldSettingsAssetUserData* CaptureSettings = FMediaProfileEditor::FindOrAddMediaFrameworkCaptureSettings();
 	if (!CaptureSettings)
@@ -586,6 +593,8 @@ void FMediaProfileMediaOutputDetailCustomization::SetMediaObject(UMediaOutput* I
 		NewCameraOutputInfo.MediaOutput = InNewMediaObject;
 		CaptureSettings->ViewportCaptures.Add(NewCameraOutputInfo);
 	}
+	
+	PinnedMediaProfile->SetMediaOutput(Index, InNewMediaObject);
 }
 
 void FMediaProfileMediaOutputDetailCustomization::RegisterCaptureSection(FPropertyEditorModule& PropertyModule, const FName& CaptureCategory)
diff --git a/Engine/Plugins/Media/MediaFrameworkUtilities/Source/MediaFrameworkUtilitiesEditor/Private/AssetEditor/SMediaProfileSourcesTreeView.cpp b/Engine/Plugins/Media/MediaFrameworkUtilities/Source/MediaFrameworkUtilitiesEditor/Private/AssetEditor/SMediaProfileSourcesTreeView.cpp
--- a/Engine/Plugins/Media/MediaFrameworkUtilities/Source/MediaFrameworkUtilitiesEditor/Private/AssetEditor/SMediaProfileSourcesTreeView.cpp
+++ b/Engine/Plugins/Media/MediaFrameworkUtilities/Source/MediaFrameworkUtilitiesEditor/Private/AssetEditor/SMediaProfileSourcesTreeView.cpp
@@ -1016,7 +1016,9 @@ void SMediaProfileSourcesTreeView::AddNewMediaSource()
 
 		MediaSourcesTreeView::TrySaveMediaProfileSettings();
 
-		FillMediaTree();
+		TOptional<TArray<int32>> MediaSourcesToSelect = TArray<int32> { NumSources };
+		TOptional<TArray<int32>> MediaOutputsToSelect = TOptional<TArray<int32>>(); 
+		FillMediaTree(MediaSourcesToSelect, MediaOutputsToSelect);
 	}
 }
 
@@ -1040,11 +1042,13 @@ void SMediaProfileSourcesTreeView::AddNewMediaOutput()
 
 		MediaSourcesTreeView::TrySaveMediaProfileSettings();
 
-		FillMediaTree();
+		TOptional<TArray<int32>> MediaSourcesToSelect = TOptional<TArray<int32>>(); 
+		TOptional<TArray<int32>> MediaOutputsToSelect = TArray<int32> { NumOutputs };
+		FillMediaTree(MediaSourcesToSelect, MediaOutputsToSelect);
 	}
 }
 
-void SMediaProfileSourcesTreeView::FillMediaTree()
+void SMediaProfileSourcesTreeView::FillMediaTree(const TOptional<TArray<int32>>& InMediaSourcesToSelect, const TOptional<TArray<int32>>& InMediaOutputsToSelect)
 {
 	// In case a media object changed type, store the current selection through their indexes to re-select them correctly after
 	// the tree has been recreated
@@ -1153,7 +1157,8 @@ void SMediaProfileSourcesTreeView::FillMediaTree()
 
 		if (MediaTreeItems.Num() > 0)
 		{
-			for (int32 SourceIndex : SelectedMediaSources)
+			const TArray<int32>& SourcesToSelect = InMediaSourcesToSelect.IsSet() ? InMediaSourcesToSelect.GetValue() : SelectedMediaSources;
+			for (int32 SourceIndex : SourcesToSelect)
 			{
 				if (SourceIndex == INDEX_NONE)
 				{
@@ -1165,7 +1170,8 @@ void SMediaProfileSourcesTreeView::FillMediaTree()
 				}
 			}
 
-			for (int32 OutputIndex : SelectedMediaOutputs)
+			const TArray<int32>& OutputsToSelect = InMediaOutputsToSelect.IsSet() ? InMediaOutputsToSelect.GetValue() : SelectedMediaOutputs;
+			for (int32 OutputIndex : OutputsToSelect)
 			{
 				if (OutputIndex == INDEX_NONE)
 				{
diff --git a/Engine/Plugins/Media/MediaFrameworkUtilities/Source/MediaFrameworkUtilitiesEditor/Private/AssetEditor/SMediaProfileSourcesTreeView.h b/Engine/Plugins/Media/MediaFrameworkUtilities/Source/MediaFrameworkUtilitiesEditor/Private/AssetEditor/SMediaProfileSourcesTreeView.h
--- a/Engine/Plugins/Media/MediaFrameworkUtilities/Source/MediaFrameworkUtilitiesEditor/Private/AssetEditor/SMediaProfileSourcesTreeView.h
+++ b/Engine/Plugins/Media/MediaFrameworkUtilities/Source/MediaFrameworkUtilitiesEditor/Private/AssetEditor/SMediaProfileSourcesTreeView.h
@@ -87,7 +87,7 @@ private:
 	void AddNewMediaOutput();
 
 	/** Recreates the tree view's source list from the current media profile's list of sources and outputs */
-	void FillMediaTree();
+	void FillMediaTree(const TOptional<TArray<int32>>& InMediaSourcesToSelect = TOptional<TArray<int32>>(), const TOptional<TArray<int32>>& InMediaOutputsToSelect = TOptional<TArray<int32>>());
 
 	/** Converts the selected tree items to a list of indexes of the selected media sources and outputs in the media profile */
 	void GetSelectedMediaIndexes(TArray<int32>& OutSelectedMediaSources, TArray<int32>& OutSelectedMediaOutputs, bool bIncludeGroupIndex = false);
diff --git a/Engine/Plugins/Media/MediaFrameworkUtilities/Source/MediaFrameworkUtilitiesEditor/Private/AssetTypeActions/AssetTypeActions_MediaProfile.cpp b/Engine/Plugins/Media/MediaFrameworkUtilities/Source/MediaFrameworkUtilitiesEditor/Private/AssetTypeActions/AssetTypeActions_MediaProfile.cpp
--- a/Engine/Plugins/Media/MediaFrameworkUtilities/Source/MediaFrameworkUtilitiesEditor/Private/AssetTypeActions/AssetTypeActions_MediaProfile.cpp
+++ b/Engine/Plugins/Media/MediaFrameworkUtilities/Source/MediaFrameworkUtilitiesEditor/Private/AssetTypeActions/AssetTypeActions_MediaProfile.cpp
@@ -3,18 +3,10 @@
 #include "AssetTypeActions/AssetTypeActions_MediaProfile.h"
 
 #include "AssetEditor/MediaProfileEditor.h"
-#include "AssetEditor/MediaProfileEditorToolkit.h"
 #include "Profile/MediaProfile.h"
 
 #define LOCTEXT_NAMESPACE "MediaProfileEditor"
 
-/** Feature flag for the new media profile editor */
-static TAutoConsoleVariable<bool> CVarUseNewMediaProfileEditor(
-	TEXT("MediaProfile.UseNewMediaProfileEditor"), false,
-	TEXT("(Experimental) If true, the new media profile editor will be used when opening media profile assets.")
-	TEXT("If false (default), the existing media profile editor will be used."),
-	ECVF_Default);
-
 FText FAssetTypeActions_MediaProfile::GetName() const
 {
 	return LOCTEXT("AssetTypeActions_MediaProfile", "Media Profile");
@@ -31,15 +23,7 @@ void FAssetTypeActions_MediaProfile::OpenAssetEditor(const TArray<UObject*>& InO
 	{
 		if (UMediaProfile* Asset = Cast<UMediaProfile>(Obj))
 		{
-			const bool bUseNewMediaProfileEditor = CVarUseNewMediaProfileEditor.GetValueOnAnyThread();
-			if (bUseNewMediaProfileEditor)
-			{
-				FMediaProfileEditor::CreateMediaProfileEditor(EToolkitMode::Standalone, EditWithinLevelEditor, Asset);
-			}
-			else
-			{
-				FMediaProfileEditorToolkit::CreateEditor(EToolkitMode::Standalone, EditWithinLevelEditor, Asset);
-			}
+			FMediaProfileEditor::CreateMediaProfileEditor(EToolkitMode::Standalone, EditWithinLevelEditor, Asset);
 		}
 	}
 }
