Media Profile: Fixed issue where viewport media capture was not setting viewport resolution based on media output settings

diff --git a/Engine/Plugins/Media/MediaFrameworkUtilities/Source/MediaFrameworkUtilitiesEditor/Private/AssetEditor/SMediaProfileMediaItemDisplay.cpp b/Engine/Plugins/Media/MediaFrameworkUtilities/Source/MediaFrameworkUtilitiesEditor/Private/AssetEditor/SMediaProfileMediaItemDisplay.cpp
--- a/Engine/Plugins/Media/MediaFrameworkUtilities/Source/MediaFrameworkUtilitiesEditor/Private/AssetEditor/SMediaProfileMediaItemDisplay.cpp
+++ b/Engine/Plugins/Media/MediaFrameworkUtilities/Source/MediaFrameworkUtilitiesEditor/Private/AssetEditor/SMediaProfileMediaItemDisplay.cpp
@@ -11,6 +11,7 @@
 #include "CaptureTab/SMediaFrameworkCapture.h"
 #include "Engine/TextureRenderTarget2D.h"
 #include "Profile/MediaProfilePlaybackManager.h"
+#include "Slate/SceneViewport.h"
 #include "UI/MediaFrameworkUtilitiesEditorStyle.h"
 #include "ViewportToolbar/UnrealEdViewportToolbar.h"
 #include "Widgets/SMediaImage.h"
@@ -277,6 +278,16 @@ public:
 		Refresh(InLockedActors);
 	}
 
+	virtual void Tick(const FGeometry& AllottedGeometry, const double InCurrentTime, const float InDeltaTime) override
+	{
+		if (PendingViewportSizeChange.IsSet() && ViewportCapture.IsValid())
+		{
+			const FIntPoint Size = PendingViewportSizeChange.GetValue();
+			ViewportCapture->Viewport->SetFixedViewportSize(Size.X, Size.Y);
+			PendingViewportSizeChange.Reset();
+		}
+	}
+	
 	/** Updates the display if the capture's locked actors have been changed */
 	void Refresh(const TArray<TLazyObjectPtr<AActor>>& InLockedActors)
 	{
@@ -330,6 +341,12 @@ public:
 
 		ViewportCapture->ViewportClient->SetActorLock(ActorToLockTo);
 	}
+
+	/** Queues up a desired viewport size that will be set on the next tick */
+	void QueueViewportSizeChange(FIntPoint InDesiredViewportSize)
+	{
+		PendingViewportSizeChange = InDesiredViewportSize;
+	}
 	
 private:
 	const FSlateBrush* GetLiveCaptureIcon() const
@@ -371,6 +388,13 @@ private:
 
 	/** Attribute for querying if the viewport capture is currently actively being captured */
 	TAttribute<bool> IsCapturing;
+
+	/**
+	 * To manually set the viewport size on the capture, we need the SViewport widget to be fully in the slate hierarchy
+	 * (as FSceneViewport needs to be able to find an SWindow from SViewport), so this optional stores any pending size change
+	 * that should be applied on the next tick of the widget, which by then should be properly in the window hierarchy
+	 */
+	TOptional<FIntPoint> PendingViewportSizeChange;
 };
 
 /** Widget to display a render target media capture output */
@@ -541,28 +565,34 @@ void SMediaProfileMediaOutputDisplay::ConfigureMediaImage()
 		if (ViewportOutputInfoIndex != INDEX_NONE)
 		{
 			TSharedPtr<FMediaProfileEditor::FViewportCapture> ViewportCapture = MediaProfileEditor.IsValid() ? MediaProfileEditor.Pin()->FindOrCreateViewportCapture(GetMediaItem()) : nullptr;
-			
-			FIntPoint TargetSize = MediaOutput->GetRequestedSize();
-			if (TargetSize == UMediaOutput::RequestCaptureSourceSize)
+			if (ViewportCapture.IsValid())
 			{
-				TargetSize = GetMutableDefault<UMediaFrameworkMediaCaptureSettings>()->DefaultViewportCaptureSize;
-			}
-
-			const float AspectRatio = (TargetSize.X > 0 && TargetSize.Y > 0) ? (float)TargetSize.X / (float)TargetSize.Y : 1.77777777f;
-			
-			MediaImageContainer->SetContent(
-				SNew(SBox)
-				.MinAspectRatio(AspectRatio)
-				.MaxAspectRatio(AspectRatio)
-				[
-					SAssignNew(ImageWidget, SMediaCaptureViewportImage, ViewportCapture, CaptureSettings->ViewportCaptures[ViewportOutputInfoIndex].LockedActors)
-					.IsCapturing(this, &SMediaProfileMediaOutputDisplay::IsCapturingMediaOutput)
-				]
-			);
+				FIntPoint TargetSize = MediaOutput->GetRequestedSize();
+				if (TargetSize == UMediaOutput::RequestCaptureSourceSize)
+				{
+					TargetSize = GetMutableDefault<UMediaFrameworkMediaCaptureSettings>()->DefaultViewportCaptureSize;
+				}
 
-			FShowFlagMenuCommands::Get().BindCommands(*ShowFlagsCommandList.Get(), ViewportCapture->ViewportClient);
-			bDisplayShowFlags = true;
-			return;
+				const float AspectRatio = (TargetSize.X > 0 && TargetSize.Y > 0) ? (float)TargetSize.X / (float)TargetSize.Y : 1.77777777f;
+
+				TSharedPtr<SMediaCaptureViewportImage> ViewportImage;
+				MediaImageContainer->SetContent(
+					SNew(SBox)
+					.MinAspectRatio(AspectRatio)
+					.MaxAspectRatio(AspectRatio)
+					[
+						SAssignNew(ViewportImage, SMediaCaptureViewportImage, ViewportCapture, CaptureSettings->ViewportCaptures[ViewportOutputInfoIndex].LockedActors)
+						.IsCapturing(this, &SMediaProfileMediaOutputDisplay::IsCapturingMediaOutput)
+					]
+				);
+
+				ViewportImage->QueueViewportSizeChange(TargetSize);
+				ImageWidget = ViewportImage;
+				
+				FShowFlagMenuCommands::Get().BindCommands(*ShowFlagsCommandList.Get(), ViewportCapture->ViewportClient);
+				bDisplayShowFlags = true;
+				return;
+			}
 		}
 
 		RenderTargetOutputInfoIndex = CaptureSettings->RenderTargetCaptures.IndexOfByPredicate(
