Media Profile: Added manual refresh button to media profile editor that closes and reopens selected media sources and outputs, as well as hooks to automatically refresh media items when property changes on them are detected

diff --git a/Engine/Plugins/Media/MediaFrameworkUtilities/Source/MediaFrameworkUtilities/Private/Profile/MediaProfilePlaybackManager.cpp b/Engine/Plugins/Media/MediaFrameworkUtilities/Source/MediaFrameworkUtilities/Private/Profile/MediaProfilePlaybackManager.cpp
--- a/Engine/Plugins/Media/MediaFrameworkUtilities/Source/MediaFrameworkUtilities/Private/Profile/MediaProfilePlaybackManager.cpp
+++ b/Engine/Plugins/Media/MediaFrameworkUtilities/Source/MediaFrameworkUtilities/Private/Profile/MediaProfilePlaybackManager.cpp
@@ -154,6 +154,53 @@ UMediaTexture* UMediaProfilePlaybackManager::OpenSource(UMediaSource* InMediaSou
 	return OpenSourceFromIndex(MediaSourceIndex, InConsumer);
 }
 
+bool UMediaProfilePlaybackManager::IsSourceOpenFromIndex(int32 InMediaSourceIndex, void* InConsumer) const
+{
+	if (!MediaSourceTextures.IsValidIndex(InMediaSourceIndex))
+	{
+		return false;
+	}
+	
+	UMediaProfile* MediaProfile = GetOwningMediaProfile();
+	if (!MediaProfile)
+	{
+		return false;
+	}
+	
+	if (InMediaSourceIndex < 0 || InMediaSourceIndex >= MediaProfile->NumMediaSources())
+	{
+		return false;
+	}
+
+	if (!MediaSourcePlayers.Contains(InMediaSourceIndex))
+	{
+		return false;
+	}
+
+	const TObjectPtr<UMediaPlayer>& MediaPlayer = MediaSourcePlayers[InMediaSourceIndex];
+	const bool bIsOpen = !MediaPlayer->IsClosed();
+
+	if (InConsumer != nullptr)
+	{
+		const bool bHasConsumer = SourceConsumers.Contains(MediaPlayer) && SourceConsumers[MediaPlayer].Contains(InConsumer);
+		return bIsOpen && bHasConsumer;
+	}
+	
+	return bIsOpen;
+}
+
+bool UMediaProfilePlaybackManager::IsSourceOpen(UMediaSource* InMediaSource, void* InConsumer) const
+{
+	UMediaProfile* MediaProfile = GetOwningMediaProfile();
+	if (!MediaProfile)
+	{
+		return false;
+	}
+	
+	int32 MediaSourceIndex = MediaProfile->FindMediaSourceIndex(InMediaSource);
+	return IsSourceOpenFromIndex(MediaSourceIndex, InConsumer);
+}
+
 void UMediaProfilePlaybackManager::CloseSourceFromIndex(int32 InMediaSourceIndex, const FCloseSourceArgs& InArgs)
 {
 	if (!MediaSourcePlayers.Contains(InMediaSourceIndex))
diff --git a/Engine/Plugins/Media/MediaFrameworkUtilities/Source/MediaFrameworkUtilities/Public/Profile/MediaProfilePlaybackManager.h b/Engine/Plugins/Media/MediaFrameworkUtilities/Source/MediaFrameworkUtilities/Public/Profile/MediaProfilePlaybackManager.h
--- a/Engine/Plugins/Media/MediaFrameworkUtilities/Source/MediaFrameworkUtilities/Public/Profile/MediaProfilePlaybackManager.h
+++ b/Engine/Plugins/Media/MediaFrameworkUtilities/Source/MediaFrameworkUtilities/Public/Profile/MediaProfilePlaybackManager.h
@@ -59,6 +59,22 @@ public:
 	 */
 	UMediaTexture* OpenSource(UMediaSource* InMediaSource, void* InConsumer = nullptr);
 
+	/**
+	 * Gets whether a media source is currently open
+	 * @param InMediaSourceIndex The index of the media source within the media profile
+	 * @param InConsumer If not null, additionally checks if the specified consumer is using the media source
+	 * @return true if the media source is open (and being used by InConsumer, if provided); otherwise false
+	 */
+	bool IsSourceOpenFromIndex(int32 InMediaSourceIndex, void* InConsumer = nullptr) const;
+
+	/**
+	 * Gets whether a media source is currently open
+	 * @param InMediaSource The media source within the media profile
+	 * @param InConsumer If not null, additionally checks if the specified consumer is using the media source
+	 * @return true if the media source is open (and being used by InConsumer, if provided); otherwise false
+	 */
+	bool IsSourceOpen(UMediaSource* InMediaSource, void* InConsumer = nullptr) const;
+	
 	/** Arguments that control what happens when a media source is closed */
 	struct FCloseSourceArgs
 	{
diff --git a/Engine/Plugins/Media/MediaFrameworkUtilities/Source/MediaFrameworkUtilitiesEditor/Private/AssetEditor/MediaProfileEditor.cpp b/Engine/Plugins/Media/MediaFrameworkUtilities/Source/MediaFrameworkUtilitiesEditor/Private/AssetEditor/MediaProfileEditor.cpp
--- a/Engine/Plugins/Media/MediaFrameworkUtilities/Source/MediaFrameworkUtilitiesEditor/Private/AssetEditor/MediaProfileEditor.cpp
+++ b/Engine/Plugins/Media/MediaFrameworkUtilities/Source/MediaFrameworkUtilitiesEditor/Private/AssetEditor/MediaProfileEditor.cpp
@@ -175,6 +175,7 @@ void FMediaProfileEditor::Initialize(const EToolkitMode::Type Mode, const TShare
 	FEditorDelegates::OnAssetsDeleted.AddSP(this, &FMediaProfileEditor::OnAssetsDeleted);
 	FCoreUObjectDelegates::OnPreObjectPropertyChanged.AddSP(this, &FMediaProfileEditor::OnObjectPreEditChange);
 	FCoreUObjectDelegates::OnObjectModified.AddSP(this, &FMediaProfileEditor::OnObjectModified);
+	FCoreUObjectDelegates::OnObjectPropertyChanged.AddSP(this, &FMediaProfileEditor::OnObjectPropertyChanged);
 	
 	const TSharedRef<FTabManager::FLayout> StandaloneDefaultLayout = FTabManager::NewLayout("Standalone_MediaProfileEditor_Layout")
 		->AddArea
@@ -230,6 +231,7 @@ void FMediaProfileEditor::Initialize(const EToolkitMode::Type Mode, const TShare
 
 FMediaProfileEditor::~FMediaProfileEditor()
 {
+	FCoreUObjectDelegates::OnObjectPropertyChanged.RemoveAll(this);
 	FCoreUObjectDelegates::OnObjectModified.RemoveAll(this);
 	FCoreUObjectDelegates::OnPreObjectPropertyChanged.RemoveAll(this);
 	FEditorDelegates::OnAssetsDeleted.RemoveAll(this);
@@ -574,6 +576,7 @@ TSharedRef<SDockTab> FMediaProfileEditor::SpawnTab_Details(const FSpawnTabArgs&
 		.Label(LOCTEXT("DetailsTabLabel", "Details"))
 		[
 			SAssignNew(DetailsPanel, SMediaProfileDetailsPanel, SharedThis(this), MediaProfileBeingEdited)
+			.OnRefresh(this, &FMediaProfileEditor::RefreshSelectedMediaItems)
 		];
 }
 
@@ -605,6 +608,42 @@ void FMediaProfileEditor::ExtendToolbar()
 	AddToolbarExtender(ToolbarExtender);
 }
 
+void FMediaProfileEditor::RefreshSelectedMediaItems(const TArray<int32>& InMediaSources,const TArray<int32>& InMediaOutputs)
+{
+	if (!IsValid(MediaProfileBeingEdited))
+	{
+		return;
+	}
+
+	for (int32 Index = 0; Index < InMediaSources.Num(); ++Index)
+	{
+		const int32 SrcIndex = InMediaSources[Index];
+
+		if (MediaProfileBeingEdited->GetPlaybackManager()->IsSourceOpenFromIndex(SrcIndex))
+		{
+			UMediaProfilePlaybackManager::FCloseSourceArgs Args;
+			Args.Consumer = this;
+			Args.bForceClose = true;
+		
+			MediaProfileBeingEdited->GetPlaybackManager()->CloseSourceFromIndex(SrcIndex, Args);
+			MediaProfileBeingEdited->GetPlaybackManager()->OpenSourceFromIndex(SrcIndex, this);
+		}
+	}
+
+	for (int32 Index = 0; Index < InMediaOutputs.Num(); ++Index)
+	{
+		const int32 OutputIndex = InMediaOutputs[Index];
+		if (UMediaOutput* MediaOutput = MediaProfileBeingEdited->GetMediaOutput(OutputIndex))
+		{
+			if (IsMediaOutputCapturing(MediaOutput))
+			{
+				StopCapturingMediaOutput(MediaOutput);
+				StartCapturingMediaOutput(MediaOutput);
+			}
+		}
+	}
+}
+
 void FMediaProfileEditor::OnMediaItemDeleted(UClass* InMediaType, int32 InMediaItemIndex)
 {
 	if (InMediaType->IsChildOf<UMediaSource>())
@@ -799,6 +838,52 @@ void FMediaProfileEditor::OnObjectModified(UObject* Object)
 	}
 }
 
+void FMediaProfileEditor::OnObjectPropertyChanged(UObject* InObject, FPropertyChangedEvent& InPropertyChangedEvent)
+{
+	if (!IsValid(MediaProfileBeingEdited))
+	{
+		return;
+	}
+	
+	if (!InObject)
+	{
+		return;
+	}
+
+	if (MediaProfileBeingEdited != InObject)
+	{
+		if (UMediaSource* MediaSource = Cast<UMediaSource>(InObject))
+		{
+			const int32 SrcIndex = MediaProfileBeingEdited->FindMediaSourceIndex(MediaSource);
+			if (SrcIndex != INDEX_NONE)
+			{
+				if (MediaProfileBeingEdited->GetPlaybackManager()->IsSourceOpenFromIndex(SrcIndex))
+				{
+					UMediaProfilePlaybackManager::FCloseSourceArgs Args;
+					Args.Consumer = this;
+					Args.bForceClose = true;
+		
+					MediaProfileBeingEdited->GetPlaybackManager()->CloseSourceFromIndex(SrcIndex, Args);
+					MediaProfileBeingEdited->GetPlaybackManager()->OpenSourceFromIndex(SrcIndex, this);
+				}
+			}
+		}
+
+		if (UMediaOutput* MediaOutput = Cast<UMediaOutput>(InObject))
+		{
+			const int32 OutputIndex = MediaProfileBeingEdited->FindMediaOutputIndex(MediaOutput);
+			if (OutputIndex != INDEX_NONE)
+			{
+				if (IsMediaOutputCapturing(MediaOutput))
+				{
+					StopCapturingMediaOutput(MediaOutput);
+					StartCapturingMediaOutput(MediaOutput);
+				}
+			}
+		}
+	}	
+}
+
 TSharedPtr<FMediaProfileEditor::FViewportCapture> FMediaProfileEditor::FindOrCreateViewportCaptureImpl(int32 ViewportOutputInfoIndex)
 {
 	// Take the opportunity to clear any invalid viewport captures from the list
diff --git a/Engine/Plugins/Media/MediaFrameworkUtilities/Source/MediaFrameworkUtilitiesEditor/Private/AssetEditor/MediaProfileEditor.h b/Engine/Plugins/Media/MediaFrameworkUtilities/Source/MediaFrameworkUtilitiesEditor/Private/AssetEditor/MediaProfileEditor.h
--- a/Engine/Plugins/Media/MediaFrameworkUtilities/Source/MediaFrameworkUtilitiesEditor/Private/AssetEditor/MediaProfileEditor.h
+++ b/Engine/Plugins/Media/MediaFrameworkUtilities/Source/MediaFrameworkUtilitiesEditor/Private/AssetEditor/MediaProfileEditor.h
@@ -105,6 +105,9 @@ private:
 	TSharedRef<SDockTab> SpawnTab_Timecode(const FSpawnTabArgs& InArgs);
 
 	void ExtendToolbar();
+
+	/** Refreshes the selected media items, reopening media sources and outputs */
+	void RefreshSelectedMediaItems(const TArray<int32>& InMediaSources,const TArray<int32>& InMediaOutputs);
 	
 	void OnMediaItemDeleted(UClass* InMediaType, int32 InMediaItemIndex);
 	void OnSelectedMediaItemsChanged(const TArray<int32>& SelectedMediaSources, const TArray<int32>& SelectedMediaOutputs);
@@ -120,6 +123,7 @@ private:
 	void OnAssetsDeleted(const TArray<UClass*>& DeletedAssetClasses);
 	void OnObjectPreEditChange(UObject* Object, const FEditPropertyChain& PropertyChain);
 	void OnObjectModified(UObject* Object);
+	void OnObjectPropertyChanged(UObject* InObject, FPropertyChangedEvent& InPropertyChangedEvent);
 
 	/** Attempts to find an existing viewport capture for the specified viewport output info, or creates one if one does not exist */
 	TSharedPtr<FViewportCapture> FindOrCreateViewportCaptureImpl(int32 ViewportOutputInfoIndex);
diff --git a/Engine/Plugins/Media/MediaFrameworkUtilities/Source/MediaFrameworkUtilitiesEditor/Private/AssetEditor/SMediaProfileDetailsPanel.cpp b/Engine/Plugins/Media/MediaFrameworkUtilities/Source/MediaFrameworkUtilitiesEditor/Private/AssetEditor/SMediaProfileDetailsPanel.cpp
--- a/Engine/Plugins/Media/MediaFrameworkUtilities/Source/MediaFrameworkUtilitiesEditor/Private/AssetEditor/SMediaProfileDetailsPanel.cpp
+++ b/Engine/Plugins/Media/MediaFrameworkUtilities/Source/MediaFrameworkUtilitiesEditor/Private/AssetEditor/SMediaProfileDetailsPanel.cpp
@@ -12,9 +12,11 @@
 #include "MediaOutput.h"
 #include "MediaProfileMediaItemDetailCustomization.h"
 #include "SMediaProfileDetailsInfoPanel.h"
+#include "SPositiveActionButton.h"
 #include "CaptureTab/SMediaFrameworkCapture.h"
 #include "Modules/ModuleManager.h"
 #include "Profile/MediaProfile.h"
+#include "UI/MediaFrameworkUtilitiesEditorStyle.h"
 #include "Widgets/Input/SComboButton.h"
 #include "Widgets/Input/SEditableTextBox.h"
 #include "Widgets/Layout/SSplitter.h"
@@ -25,6 +27,7 @@ void SMediaProfileDetailsPanel::Construct(const FArguments& InArgs, const TShare
 {
 	OwningEditor = InOwningEditor;
 	MediaProfile = InMediaProfile;
+	OnRefresh = InArgs._OnRefresh;
 
 	FPropertyEditorModule& PropertyModule = FModuleManager::LoadModuleChecked<FPropertyEditorModule>("PropertyEditor");
 	FDetailsViewArgs DetailsArgs;
@@ -77,6 +80,28 @@ void SMediaProfileDetailsPanel::Construct(const FArguments& InArgs, const TShare
 				[
 					DetailsView->GetFilterAreaWidget().ToSharedRef()
 				]
+				
+				+SVerticalBox::Slot()
+				.AutoHeight()
+				.HAlign(HAlign_Left)
+				.VAlign(VAlign_Center)
+				.Padding(6.0f)
+				[
+					SNew(SPositiveActionButton)
+					.Icon(FMediaFrameworkUtilitiesEditorStyle::Get().GetBrush("ToolbarIcon.Refresh"))
+					.Text(LOCTEXT("RefreshButtonLabel", "Refresh Media Properties"))
+					.ToolTipText(LOCTEXT("RefreshButtonTooltip", "Refreshes the streams of the selected media items to match the item's current property values"))
+					.Visibility_Lambda([this]
+					{
+						return !SelectedMediaSources.IsEmpty() || !SelectedMediaOutputs.IsEmpty() ? EVisibility::Visible : EVisibility::Hidden;
+					})
+					.OnClicked_Lambda([this]
+					{
+						OnRefresh.ExecuteIfBound(SelectedMediaSources, SelectedMediaOutputs);
+						return FReply::Handled();
+					})
+				]
+				
 				+SVerticalBox::Slot()
 				[
 					DetailsView.ToSharedRef()
diff --git a/Engine/Plugins/Media/MediaFrameworkUtilities/Source/MediaFrameworkUtilitiesEditor/Private/AssetEditor/SMediaProfileDetailsPanel.h b/Engine/Plugins/Media/MediaFrameworkUtilities/Source/MediaFrameworkUtilitiesEditor/Private/AssetEditor/SMediaProfileDetailsPanel.h
--- a/Engine/Plugins/Media/MediaFrameworkUtilities/Source/MediaFrameworkUtilitiesEditor/Private/AssetEditor/SMediaProfileDetailsPanel.h
+++ b/Engine/Plugins/Media/MediaFrameworkUtilities/Source/MediaFrameworkUtilitiesEditor/Private/AssetEditor/SMediaProfileDetailsPanel.h
@@ -17,10 +17,12 @@ class SMediaProfileDetailsPanel : public SCompoundWidget
 {
 public:
 	DECLARE_DELEGATE(FOnDetailsChanged);
+	DECLARE_DELEGATE_TwoParams(FOnRefreshMediaItems, const TArray<int32>& /* InMediaSources */, const TArray<int32>& /* InMediaOutputs */)
 	
 public:
 	SLATE_BEGIN_ARGS(SMediaProfileDetailsPanel) {}
 		SLATE_EVENT(FOnDetailsChanged, OnDetailsChanged)
+		SLATE_EVENT(FOnRefreshMediaItems, OnRefresh)
 	SLATE_END_ARGS()
 	
 	void Construct(const FArguments& InArgs, const TSharedPtr<FMediaProfileEditor> InOwningEditor, UMediaProfile* InMediaProfile);
@@ -56,4 +58,7 @@ private:
 
 	/** Indices of the selected media outputs within the media profile to display in the details panel */
 	TArray<int32> SelectedMediaOutputs;
+
+	/** Raised when the refresh button is pressed */
+	FOnRefreshMediaItems OnRefresh;
 };
diff --git a/Engine/Plugins/Media/MediaFrameworkUtilities/Source/MediaFrameworkUtilitiesEditor/Private/UI/MediaFrameworkUtilitiesEditorStyle.cpp b/Engine/Plugins/Media/MediaFrameworkUtilities/Source/MediaFrameworkUtilitiesEditor/Private/UI/MediaFrameworkUtilitiesEditorStyle.cpp
--- a/Engine/Plugins/Media/MediaFrameworkUtilities/Source/MediaFrameworkUtilitiesEditor/Private/UI/MediaFrameworkUtilitiesEditorStyle.cpp
+++ b/Engine/Plugins/Media/MediaFrameworkUtilities/Source/MediaFrameworkUtilitiesEditor/Private/UI/MediaFrameworkUtilitiesEditorStyle.cpp
@@ -10,6 +10,7 @@
 
 // This is to fix the issue that SlateStyleMacros like IMAGE_BRUSH look for RootToContentDir but StyleSet->RootToContentDir is how this style is set up
 #define RootToContentDir MediaFrameworkUtilitiesStyle::StyleInstance->RootToContentDir
+#define RootToCoreContentDir MediaFrameworkUtilitiesStyle::StyleInstance->RootToCoreContentDir
 
 namespace MediaFrameworkUtilitiesStyle
 {
@@ -27,6 +28,7 @@ namespace MediaFrameworkUtilitiesStyle
 void FMediaFrameworkUtilitiesEditorStyle::Register()
 {
 	MediaFrameworkUtilitiesStyle::StyleInstance = MakeUnique<FSlateStyleSet>(MediaFrameworkUtilitiesStyle::NAME_StyleName);
+	MediaFrameworkUtilitiesStyle::StyleInstance->SetCoreContentRoot(FPaths::EngineContentDir() / TEXT("Slate"));
 	MediaFrameworkUtilitiesStyle::StyleInstance->SetContentRoot(FPaths::EnginePluginsDir() / TEXT("Media/MediaFrameworkUtilities/Content/Editor/Icons/"));
 
 	MediaFrameworkUtilitiesStyle::StyleInstance->Set("ClassThumbnail.MediaBundle", new IMAGE_BRUSH("MediaBundle_64x", MediaFrameworkUtilitiesStyle::Icon64x64));
@@ -53,6 +55,8 @@ void FMediaFrameworkUtilitiesEditorStyle::Register()
 
 	MediaFrameworkUtilitiesStyle::StyleInstance->Set("ToolbarIcon.Timecode", new IMAGE_BRUSH_SVG("Timecode", MediaFrameworkUtilitiesStyle::Icon16x16));
 	MediaFrameworkUtilitiesStyle::StyleInstance->Set("ToolbarIcon.Genlock", new IMAGE_BRUSH_SVG("Genlock", MediaFrameworkUtilitiesStyle::Icon16x16));
+
+	MediaFrameworkUtilitiesStyle::StyleInstance->Set("ToolbarIcon.Refresh", new CORE_IMAGE_BRUSH_SVG("Starship/Common/Update", MediaFrameworkUtilitiesStyle::Icon16x16));
 	
 	//Use existing icon of MaterialInstanceConstant from Engine Editor content
 	const FString EngineContentPath = FPaths::EngineContentDir() / TEXT("Editor/Slate/Icons/AssetIcons/MaterialInstanceConstant_64x.png");
@@ -81,3 +85,5 @@ const ISlateStyle& FMediaFrameworkUtilitiesEditorStyle::Get()
 }
 
 #undef IMAGE_BRUSH
+#undef RootToContentDir
+#undef RootToCoreContentDir
\ No newline at end of file
