Media Profile Editor: Added missing Current Viewport media capture option

diff --git a/Engine/Plugins/Media/MediaFrameworkUtilities/Source/MediaFrameworkUtilitiesEditor/Private/AssetEditor/MediaProfileEditor.cpp b/Engine/Plugins/Media/MediaFrameworkUtilities/Source/MediaFrameworkUtilitiesEditor/Private/AssetEditor/MediaProfileEditor.cpp
--- a/Engine/Plugins/Media/MediaFrameworkUtilities/Source/MediaFrameworkUtilitiesEditor/Private/AssetEditor/MediaProfileEditor.cpp
+++ b/Engine/Plugins/Media/MediaFrameworkUtilities/Source/MediaFrameworkUtilitiesEditor/Private/AssetEditor/MediaProfileEditor.cpp
@@ -3,6 +3,8 @@
 
 #include "MediaProfileEditor.h"
 
+#include "IAssetViewport.h"
+#include "LevelEditor.h"
 #include "Engine/TextureRenderTarget2D.h"
 #include "LevelEditorViewport.h"
 #include "PropertyEditorModule.h"
@@ -156,6 +158,56 @@ FMediaProfileEditor::FViewportCapture::~FViewportCapture()
 	ViewportWidget.Reset();
 }
 
+FMediaProfileEditor::FActiveEditorViewport::FActiveEditorViewport(const TSharedPtr<FSceneViewport>& InViewport, const TSharedPtr<IAssetViewport>& InViewportInterface)
+	: Viewport(InViewport)
+	, ViewportInterface(InViewportInterface)
+{
+	if (InViewportInterface.IsValid())
+	{
+		FEditorViewportClient& ViewportClient = InViewportInterface->GetAssetViewportClient();
+
+		// Cache the viewport's current settings, then set them to more appropriate settings for recording
+		CachedViewportFlags.bRealTime = ViewportClient.IsRealtime();
+		CachedViewportFlags.bSetListenerPosition = ViewportClient.bSetListenerPosition;
+		CachedViewportFlags.bDrawAxes = ViewportClient.bDrawAxes;
+		CachedViewportFlags.bDisableInput = ViewportClient.bDisableInput;
+		CachedViewportFlags.bAllowCinematicControl = ViewportClient.AllowsCinematicControl();
+		CachedViewportFlags.VisibilityDelegate = ViewportClient.VisibilityDelegate;
+
+		ViewportClient.SetRealtime(true);
+		ViewportClient.bSetListenerPosition = false;
+		ViewportClient.bDrawAxes = false;
+		ViewportClient.SetAllowCinematicControl(false);
+		ViewportClient.VisibilityDelegate.BindLambda([] { return true; });
+		
+		if (UMediaFrameworkWorldSettingsAssetUserData* CaptureSettings = FindOrAddMediaFrameworkCaptureSettings())
+		{
+			// Disable input by default unless we're dynamically resizing the viewport texture as we capture.
+			// Note that we originally disable input because click routing is broken when we resize the viewport through media capture at the moment.
+			const bool bDisableInput = CaptureSettings->CurrentViewportMediaOutput.CaptureOptions.ResizeMethod != EMediaCaptureResizeMethod::ResizeInRenderPass;
+			ViewportClient.bDisableInput = bDisableInput ? true : ViewportClient.bDisableInput;
+		}
+	}
+}
+
+FMediaProfileEditor::FActiveEditorViewport::~FActiveEditorViewport()
+{
+	TSharedPtr<FSceneViewport> PinnedViewport = Viewport.Pin();
+	TSharedPtr<IAssetViewport> PinnedViewportInterface = ViewportInterface.Pin();
+	if (PinnedViewportInterface.IsValid() && PinnedViewportInterface->GetSharedActiveViewport() == PinnedViewport)
+	{
+		FEditorViewportClient& ViewportClient = PinnedViewportInterface->GetAssetViewportClient();
+
+		// Reset settings
+		ViewportClient.SetRealtime(CachedViewportFlags.bRealTime);
+		ViewportClient.bSetListenerPosition = CachedViewportFlags.bSetListenerPosition;
+		ViewportClient.bDrawAxes = CachedViewportFlags.bDrawAxes;
+		ViewportClient.bDisableInput = CachedViewportFlags.bDisableInput;
+		ViewportClient.SetAllowCinematicControl(CachedViewportFlags.bAllowCinematicControl);
+		ViewportClient.VisibilityDelegate = CachedViewportFlags.VisibilityDelegate;
+	}
+}
+
 TSharedRef<FMediaProfileEditor> FMediaProfileEditor::CreateMediaProfileEditor(const EToolkitMode::Type Mode, const TSharedPtr<IToolkitHost>& InitToolkitHost, UMediaProfile* InMediaProfile)
 {
 	TSharedRef<FMediaProfileEditor> NewEditor = MakeShared<FMediaProfileEditor>();
@@ -176,6 +228,7 @@ void FMediaProfileEditor::Initialize(const EToolkitMode::Type Mode, const TShare
 	FCoreUObjectDelegates::OnPreObjectPropertyChanged.AddSP(this, &FMediaProfileEditor::OnObjectPreEditChange);
 	FCoreUObjectDelegates::OnObjectModified.AddSP(this, &FMediaProfileEditor::OnObjectModified);
 	FCoreUObjectDelegates::OnObjectPropertyChanged.AddSP(this, &FMediaProfileEditor::OnObjectPropertyChanged);
+	GEditor->OnLevelViewportClientListChanged().AddSP(this, &FMediaProfileEditor::OnLevelViewportClientListChanged);
 	
 	const TSharedRef<FTabManager::FLayout> StandaloneDefaultLayout = FTabManager::NewLayout("Standalone_MediaProfileEditor_Layout")
 		->AddArea
@@ -231,6 +284,7 @@ void FMediaProfileEditor::Initialize(const EToolkitMode::Type Mode, const TShare
 
 FMediaProfileEditor::~FMediaProfileEditor()
 {
+	GEditor->OnLevelViewportClientListChanged().RemoveAll(this);
 	FCoreUObjectDelegates::OnObjectPropertyChanged.RemoveAll(this);
 	FCoreUObjectDelegates::OnObjectModified.RemoveAll(this);
 	FCoreUObjectDelegates::OnPreObjectPropertyChanged.RemoveAll(this);
@@ -320,6 +374,24 @@ void FMediaProfileEditor::CloseAllMediaSources()
 	}
 }
 
+TSharedPtr<FSceneViewport> FMediaProfileEditor::GetActiveEditorViewport() const
+{
+	// If there is a live capture of the active editor viewport, return the viewport it is using
+	for (const FLiveMediaCapture& Capture : LiveCaptures)
+	{
+		if (Capture.ActiveEditorViewport.IsValid())
+		{
+			return Capture.ActiveEditorViewport->Viewport.Pin();
+		}
+	}
+
+	// Otherwise, attempt to find a viewport
+	TSharedPtr<FSceneViewport> Viewport;
+	TSharedPtr<IAssetViewport> ViewportInterface;
+	GetActiveEditorViewportImpl(Viewport, ViewportInterface);
+	return Viewport;
+}
+
 TSharedPtr<FMediaProfileEditor::FViewportCapture> FMediaProfileEditor::FindOrCreateViewportCapture(UMediaOutput* InMediaOutput)
 {
 	UMediaFrameworkWorldSettingsAssetUserData* CaptureSettings = FindOrAddMediaFrameworkCaptureSettings();
@@ -354,6 +426,19 @@ bool FMediaProfileEditor::IsMediaOutputCapturing(UMediaOutput* InMediaOutput) co
 	});
 }
 
+TOptional<EMediaCaptureState> FMediaProfileEditor::GetMediaOutputCaptureState(UMediaOutput* InMediaOutput) const
+{
+	for (const FLiveMediaCapture& Capture : LiveCaptures)
+	{
+		if (Capture.MediaOutput == InMediaOutput)
+		{
+			return Capture.CaptureState;
+		}
+	}
+
+	return TOptional<EMediaCaptureState>();
+}
+
 bool FMediaProfileEditor::CanMediaOutputCapture(UMediaOutput* InMediaOutput) const
 {
 	UMediaFrameworkWorldSettingsAssetUserData* CaptureSettings = FindOrAddMediaFrameworkCaptureSettings();
@@ -364,6 +449,16 @@ bool FMediaProfileEditor::CanMediaOutputCapture(UMediaOutput* InMediaOutput) con
 
 	// Return true if there is at least one valid capture configuration for the specified media output
 	bool bHasValidCaptureSettings = false;
+	if (CaptureSettings->CurrentViewportMediaOutput.MediaOutput == InMediaOutput)
+	{
+		//Check to see if there is a valid editor viewport that can be captured from
+		TSharedPtr<FSceneViewport> Viewport;
+		TSharedPtr<IAssetViewport> ViewportInterface;
+		GetActiveEditorViewportImpl(Viewport, ViewportInterface);
+
+		bHasValidCaptureSettings = Viewport.IsValid();
+	}
+	
 	for (const FMediaFrameworkCaptureCameraViewportCameraOutputInfo& ViewportOutputInfo : CaptureSettings->ViewportCaptures)
 	{
 		if (ViewportOutputInfo.MediaOutput != InMediaOutput)
@@ -414,6 +509,32 @@ void FMediaProfileEditor::StartCapturingMediaOutput(UMediaOutput* InMediaOutput)
 	}
 	
 	bool bStartedCapture = false;
+	if (CaptureSettings->CurrentViewportMediaOutput.MediaOutput == InMediaOutput)
+	{
+		FLiveMediaCapture NewLiveCapture;
+		NewLiveCapture.MediaOutput = InMediaOutput;
+		NewLiveCapture.MediaCapture = TStrongObjectPtr(InMediaOutput->CreateMediaCapture());
+
+		if (NewLiveCapture.MediaCapture.IsValid())
+		{
+			TSharedPtr<FSceneViewport> Viewport;
+			TSharedPtr<IAssetViewport> ViewportInterface;
+			GetActiveEditorViewportImpl(Viewport, ViewportInterface);
+
+			if (Viewport.IsValid())
+			{
+				NewLiveCapture.ActiveEditorViewport = MakeShared<FActiveEditorViewport>(Viewport, ViewportInterface);
+				if (NewLiveCapture.MediaCapture->CaptureSceneViewport(Viewport, CaptureSettings->CurrentViewportMediaOutput.CaptureOptions))
+				{
+					NewLiveCapture.CaptureState = NewLiveCapture.MediaCapture->GetState();
+					NewLiveCapture.MediaCapture->OnStateChangedNative.AddSP(this, &FMediaProfileEditor::OnMediaCaptureStateChanged, NewLiveCapture.MediaCapture.Get());
+					LiveCaptures.Add(NewLiveCapture);
+					bStartedCapture = true;
+				}
+			}
+		}
+	}
+	
 	for (int32 Index = 0; Index < CaptureSettings->ViewportCaptures.Num(); ++Index)
 	{
 		const FMediaFrameworkCaptureCameraViewportCameraOutputInfo& ViewportOutputInfo = CaptureSettings->ViewportCaptures[Index];
@@ -439,6 +560,8 @@ void FMediaProfileEditor::StartCapturingMediaOutput(UMediaOutput* InMediaOutput)
 
 		if (NewLiveCapture.MediaCapture->CaptureSceneViewport(NewLiveCapture.ViewportCapture->Viewport, ViewportOutputInfo.CaptureOptions))
 		{
+			NewLiveCapture.CaptureState = NewLiveCapture.MediaCapture->GetState();
+			NewLiveCapture.MediaCapture->OnStateChangedNative.AddSP(this, &FMediaProfileEditor::OnMediaCaptureStateChanged, NewLiveCapture.MediaCapture.Get());
 			LiveCaptures.Add(NewLiveCapture);
 			bStartedCapture = true;
 		}
@@ -459,6 +582,8 @@ void FMediaProfileEditor::StartCapturingMediaOutput(UMediaOutput* InMediaOutput)
 		{
 			if (NewLiveCapture.MediaCapture->CaptureTextureRenderTarget2D(RenderTargetOutputInfo.RenderTarget.Get(), RenderTargetOutputInfo.CaptureOptions))
 			{
+				NewLiveCapture.CaptureState = NewLiveCapture.MediaCapture->GetState();
+				NewLiveCapture.MediaCapture->OnStateChangedNative.AddSP(this, &FMediaProfileEditor::OnMediaCaptureStateChanged, NewLiveCapture.MediaCapture.Get());
 				LiveCaptures.Add(NewLiveCapture);
 				bStartedCapture = true;
 			}
@@ -485,15 +610,7 @@ void FMediaProfileEditor::StopCapturingMediaOutput(UMediaOutput* InMediaOutput)
 			continue;
 		}
 
-		if (Capture.MediaCapture.IsValid())
-		{
-			Capture.MediaCapture->StopCapture(false);
-		}
-
-		if (Capture.MediaOutput.IsValid())
-		{
-			Capture.MediaOutput->OnOutputModified().RemoveAll(this);
-		}
+		CloseLiveCapture(Capture);
 	}
 
 	LiveCaptures.RemoveAll([InMediaOutput](const FLiveMediaCapture& Capture)
@@ -506,15 +623,7 @@ void FMediaProfileEditor::StopCapturingAllMediaOutputs()
 {
 	for (FLiveMediaCapture& Capture : LiveCaptures)
 	{
-		if (Capture.MediaCapture.IsValid())
-		{
-			Capture.MediaCapture->StopCapture(false);
-		}
-
-		if (Capture.MediaOutput.IsValid())
-		{
-			Capture.MediaOutput->OnOutputModified().RemoveAll(this);
-		}
+		CloseLiveCapture(Capture);
 	}
 
 	LiveCaptures.Empty();
@@ -703,6 +812,40 @@ void FMediaProfileEditor::OnMediaOutputModified(UMediaOutput* InMediaOutput)
 	}
 }
 
+void FMediaProfileEditor::OnMediaCaptureStateChanged(UMediaCapture* InMediaCapture)
+{
+	if (!IsValid(InMediaCapture))
+	{
+		return;
+	}
+	
+	int32 Index = LiveCaptures.IndexOfByPredicate([InMediaCapture](const FLiveMediaCapture& InLiveCapture)
+	{
+		return InLiveCapture.MediaCapture.Get() == InMediaCapture;
+	});
+	
+	if (Index == INDEX_NONE)
+	{
+		return;
+	}
+
+	// UMediaCapture automatically transitions from Error to Stopped, but we want to cache the error state
+	// and not actually close the capture yet so that the user can be notified that the capture had an error
+	EMediaCaptureState CaptureState = InMediaCapture->GetState();
+	if (CaptureState == EMediaCaptureState::Stopped)
+	{
+		if (LiveCaptures[Index].CaptureState != EMediaCaptureState::Error)
+		{
+			CloseLiveCapture(LiveCaptures[Index]);
+			LiveCaptures.RemoveAt(Index);
+		}
+	}
+	else
+	{
+		LiveCaptures[Index].CaptureState = CaptureState;
+	}
+}
+
 void FMediaProfileEditor::OnRestartCapturesNextTick()
 {
 	for (const TWeakObjectPtr<UMediaOutput>& MediaOutput : OutputsToRestartCapture)
@@ -884,6 +1027,35 @@ void FMediaProfileEditor::OnObjectPropertyChanged(UObject* InObject, FPropertyCh
 	}	
 }
 
+void FMediaProfileEditor::OnLevelViewportClientListChanged()
+{
+	for (int32 Index = 0; Index < LiveCaptures.Num(); ++Index)
+	{
+		FLiveMediaCapture& Capture = LiveCaptures[Index];
+		if (Capture.ActiveEditorViewport.IsValid())
+		{
+			bool bStillActive = false;
+			if (TSharedPtr<FSceneViewport> PinnedViewport = Capture.ActiveEditorViewport->Viewport.Pin())
+			{
+				for (FLevelEditorViewportClient* ViewportClient : GEditor->GetLevelViewportClients())
+				{
+					if (ViewportClient->Viewport == PinnedViewport.Get())
+					{
+						bStillActive = true;
+						break;
+					}
+				}
+			}
+
+			if (!bStillActive)
+			{
+				CloseLiveCapture(Capture);
+				LiveCaptures.RemoveAt(Index--);
+			}
+		}
+	}
+}
+
 TSharedPtr<FMediaProfileEditor::FViewportCapture> FMediaProfileEditor::FindOrCreateViewportCaptureImpl(int32 ViewportOutputInfoIndex)
 {
 	// Take the opportunity to clear any invalid viewport captures from the list
@@ -938,4 +1110,59 @@ TSharedPtr<FMediaProfileEditor::FViewportCapture> FMediaProfileEditor::FindOrCre
 	return ActiveViewportCaptures[ViewportOutputInfoIndex].Pin();
 }
 
+void FMediaProfileEditor::GetActiveEditorViewportImpl(TSharedPtr<FSceneViewport>& OutSceneViewport, TSharedPtr<IAssetViewport>& OutViewportInterface) const
+{
+	OutSceneViewport = nullptr;
+	OutViewportInterface = nullptr;
+	
+	UMediaFrameworkWorldSettingsAssetUserData* CaptureSettings = FindOrAddMediaFrameworkCaptureSettings();
+	if (!CaptureSettings)
+	{
+		return;
+	}
+
+	if (!IsValid(CaptureSettings->CurrentViewportMediaOutput.MediaOutput))
+	{
+		return;
+	}
+
+	// If there is an active PIE viewport, return it
+	for (const FWorldContext& Context : GEngine->GetWorldContexts())
+	{
+		if (Context.WorldType == EWorldType::PIE)
+		{
+			UEditorEngine* EditorEngine = CastChecked<UEditorEngine>(GEngine);
+			FSlatePlayInEditorInfo& Info = EditorEngine->SlatePlayInEditorMap.FindChecked(Context.ContextHandle);
+			if (Info.SlatePlayInEditorWindowViewport.IsValid())
+			{
+				OutSceneViewport = Info.SlatePlayInEditorWindowViewport;
+			}
+		}
+	}
+
+	if (OutSceneViewport.IsValid())
+	{
+		return;
+	}
+	
+	// Find a level editor viewport
+	FLevelEditorModule& LevelEditorModule = FModuleManager::Get().GetModuleChecked<FLevelEditorModule>(TEXT("LevelEditor"));
+	OutViewportInterface = LevelEditorModule.GetFirstActiveViewport();
+	OutSceneViewport = OutViewportInterface.IsValid() ? OutViewportInterface->GetSharedActiveViewport() : nullptr;
+}
+
+void FMediaProfileEditor::CloseLiveCapture(FLiveMediaCapture& InLiveCapture)
+{
+	if (InLiveCapture.MediaCapture.IsValid())
+	{
+		InLiveCapture.MediaCapture->OnStateChangedNative.RemoveAll(this);
+		InLiveCapture.MediaCapture->StopCapture(false);
+	}
+
+	if (InLiveCapture.MediaOutput.IsValid())
+	{
+		InLiveCapture.MediaOutput->OnOutputModified().RemoveAll(this);
+	}
+}
+
 #undef LOCTEXT_NAMESPACE
diff --git a/Engine/Plugins/Media/MediaFrameworkUtilities/Source/MediaFrameworkUtilitiesEditor/Private/AssetEditor/MediaProfileEditor.h b/Engine/Plugins/Media/MediaFrameworkUtilities/Source/MediaFrameworkUtilitiesEditor/Private/AssetEditor/MediaProfileEditor.h
--- a/Engine/Plugins/Media/MediaFrameworkUtilities/Source/MediaFrameworkUtilitiesEditor/Private/AssetEditor/MediaProfileEditor.h
+++ b/Engine/Plugins/Media/MediaFrameworkUtilities/Source/MediaFrameworkUtilitiesEditor/Private/AssetEditor/MediaProfileEditor.h
@@ -3,8 +3,11 @@
 #pragma once
 
 #include "CoreMinimal.h"
+#include "EditorViewportClient.h"
+#include "ShowFlags.h"
 #include "Toolkits/AssetEditorToolkit.h"
 
+enum class EMediaCaptureState : uint8;
 class SMediaFrameworkTimecodeGenlockHeader;
 class SMediaProfileViewport;
 class FMediaFrameworkCaptureLevelEditorViewportClient;
@@ -39,6 +42,39 @@ public:
 	};
 	
 	DECLARE_MULTICAST_DELEGATE_OneParam(FOnCaptureMethodChanged, UMediaOutput*)
+
+private:
+	/** Stores objects necessary for capturing the currently active editor viewport */
+	struct FActiveEditorViewport
+	{
+		TWeakPtr<FSceneViewport> Viewport;
+		TWeakPtr<IAssetViewport> ViewportInterface;
+
+		struct
+		{
+			bool bRealTime = false;
+			bool bSetListenerPosition = false;
+			bool bDrawAxes = false;
+			bool bDisableInput = false;
+			bool bAllowCinematicControl = false;
+			FEngineShowFlags EngineShowFlags = ESFIM_Editor;
+			FEngineShowFlags LastEngineShowFlags = ESFIM_Game;
+			FViewportStateGetter VisibilityDelegate;
+		} CachedViewportFlags;
+
+		FActiveEditorViewport(const TSharedPtr<FSceneViewport>& InViewport, const TSharedPtr<IAssetViewport>& InViewportInterface);
+		~FActiveEditorViewport();
+	};
+	
+	/** Stores objects necessary for performing a live capture of a media output */
+	struct FLiveMediaCapture
+	{
+		TWeakObjectPtr<UMediaOutput> MediaOutput;
+		TStrongObjectPtr<UMediaCapture> MediaCapture;
+		TSharedPtr<FViewportCapture> ViewportCapture;
+		TSharedPtr<FActiveEditorViewport> ActiveEditorViewport;
+		EMediaCaptureState CaptureState;
+	};
 	
 public:
 	static TSharedRef<FMediaProfileEditor> CreateMediaProfileEditor(const EToolkitMode::Type Mode, const TSharedPtr< IToolkitHost >& InitToolkitHost, UMediaProfile* InMediaProfile);
@@ -66,6 +102,9 @@ public:
 
 	/** Closes all open media sources in the media profile */
 	void CloseAllMediaSources();
+
+	/** Returns a valid editor viewport that can be captured from, or null if there are none */
+	TSharedPtr<FSceneViewport> GetActiveEditorViewport() const;
 	
 	/**
 	 * Attempts to find an existing viewport capture for the specified media output, or creates one if one does not exist
@@ -77,6 +116,9 @@ public:
 	/** Determines if the specified media output is currently capturing */
 	bool IsMediaOutputCapturing(UMediaOutput* InMediaOutput) const;
 
+	/** Gets the current state of the live capture for the specified media output, or a null TOptional if there is no live capture */
+	TOptional<EMediaCaptureState> GetMediaOutputCaptureState(UMediaOutput* InMediaOutput) const;
+	
 	/** Determines if the specified media output is properly configured for capturing */
 	bool CanMediaOutputCapture(UMediaOutput* InMediaOutput) const;
 	
@@ -112,6 +154,7 @@ private:
 	void OnMediaItemDeleted(UClass* InMediaType, int32 InMediaItemIndex);
 	void OnSelectedMediaItemsChanged(const TArray<int32>& SelectedMediaSources, const TArray<int32>& SelectedMediaOutputs);
 	void OnMediaOutputModified(UMediaOutput* InMediaOutput);
+	void OnMediaCaptureStateChanged(UMediaCapture* InMediaCapture);
 	void OnRestartCapturesNextTick();
 
 
@@ -124,9 +167,16 @@ private:
 	void OnObjectPreEditChange(UObject* Object, const FEditPropertyChain& PropertyChain);
 	void OnObjectModified(UObject* Object);
 	void OnObjectPropertyChanged(UObject* InObject, FPropertyChangedEvent& InPropertyChangedEvent);
+	void OnLevelViewportClientListChanged();
 
 	/** Attempts to find an existing viewport capture for the specified viewport output info, or creates one if one does not exist */
 	TSharedPtr<FViewportCapture> FindOrCreateViewportCaptureImpl(int32 ViewportOutputInfoIndex);
+
+	/** Attempts to find an editor viewport to capture from (can be either a PIE viewport or a level editor viewport) */
+	void GetActiveEditorViewportImpl(TSharedPtr<FSceneViewport>& OutSceneViewport, TSharedPtr<IAssetViewport>& OutViewportInterface) const;
+
+	/** Closes the live media capture and performs any necessary cleanup */
+	void CloseLiveCapture(FLiveMediaCapture& InLiveCapture);
 	
 private:
 	/** The media profile being edited by this editor */
@@ -149,14 +199,6 @@ private:
 	 * the viewport captures are cleaned up, with the invalid weak pointers cleaned up at the next opportunity
 	 */
 	TMap<int32, TWeakPtr<FViewportCapture>> ActiveViewportCaptures;
-
-	/** Stores objects necessary for performing a live capture of a media output */
-	struct FLiveMediaCapture
-	{
-		TWeakObjectPtr<UMediaOutput> MediaOutput;
-		TStrongObjectPtr<UMediaCapture> MediaCapture;
-		TSharedPtr<FViewportCapture> ViewportCapture;
-	};
 	
 	/** List of any live media output captures that the media profile editor is managing */
 	TArray<FLiveMediaCapture> LiveCaptures;
diff --git a/Engine/Plugins/Media/MediaFrameworkUtilities/Source/MediaFrameworkUtilitiesEditor/Private/AssetEditor/MediaProfileMediaItemDetailCustomization.cpp b/Engine/Plugins/Media/MediaFrameworkUtilities/Source/MediaFrameworkUtilitiesEditor/Private/AssetEditor/MediaProfileMediaItemDetailCustomization.cpp
--- a/Engine/Plugins/Media/MediaFrameworkUtilities/Source/MediaFrameworkUtilitiesEditor/Private/AssetEditor/MediaProfileMediaItemDetailCustomization.cpp
+++ b/Engine/Plugins/Media/MediaFrameworkUtilities/Source/MediaFrameworkUtilitiesEditor/Private/AssetEditor/MediaProfileMediaItemDetailCustomization.cpp
@@ -96,7 +96,12 @@ public:
 			{
 				FMenuBuilder MenuBuilder(true, nullptr);
 
-				MenuBuilder.AddMenuEntry(LOCTEXT("MediaCaptureViewportMethod", "Viewport"),
+				MenuBuilder.AddMenuEntry(LOCTEXT("MediaCaptureViewportMethod", "Current Viewport"),
+					LOCTEXT("MediaCaptureViewportMethodToolTip", "Capture from the current viewport"),
+					FSlateIcon(),
+					FUIAction(FExecuteAction::CreateSP(this, &SMediaCaptureMethodComboButton::OnCurrentViewportMethodSelected)));
+				
+				MenuBuilder.AddMenuEntry(LOCTEXT("MediaCaptureViewportMethod", "Media Viewport"),
 					LOCTEXT("MediaCaptureViewportMethodToolTip", "Capture from a viewport"),
 					FSlateIcon(),
 					FUIAction(FExecuteAction::CreateSP(this, &SMediaCaptureMethodComboButton::OnViewportMethodSelected)));
@@ -118,6 +123,51 @@ public:
 	}
 
 private:
+	void OnCurrentViewportMethodSelected()
+	{
+		UMediaFrameworkWorldSettingsAssetUserData* MediaCaptureSettings = FMediaProfileEditor::FindOrAddMediaFrameworkCaptureSettings();
+		if (!MediaCaptureSettings)
+		{
+			return;
+		}
+		
+		for (const TWeakObjectPtr<UObject>& WeakObject : Objects)
+		{
+			TStrongObjectPtr<UObject> Object = WeakObject.Pin();
+			UMediaOutput* MediaOutput = Object.IsValid() ? Cast<UMediaOutput>(Object.Get()) : nullptr;
+			if (!MediaOutput)
+			{
+				continue;
+			}
+
+			if (MediaOutput->IsA<UDummyMediaOutput>())
+			{
+				continue;
+			}
+		
+			// We expect there to be only one capture configuration per media output, so clear
+			// any existing viewport and render target capture configurations before adding the new one
+			MediaCaptureSettings->ViewportCaptures.RemoveAll(GetFindOutputInfoPredicate<FMediaFrameworkCaptureCameraViewportCameraOutputInfo>(MediaOutput));
+			MediaCaptureSettings->RenderTargetCaptures.RemoveAll(GetFindOutputInfoPredicate<FMediaFrameworkCaptureRenderTargetCameraOutputInfo>(MediaOutput));
+			
+			// Now, if the object has an existing viewport or render target capture config, we want to copy that config over to the current viewport config, so find it
+			if (FMediaFrameworkCaptureCameraViewportCameraOutputInfo* ExistingViewportOutputInfo =
+				MediaCaptureSettings->ViewportCaptures.FindByPredicate(GetFindOutputInfoPredicate<FMediaFrameworkCaptureCameraViewportCameraOutputInfo>(MediaOutput)))
+			{
+				MediaCaptureSettings->CurrentViewportMediaOutput.CaptureOptions = ExistingViewportOutputInfo->CaptureOptions;
+			}
+			else if (FMediaFrameworkCaptureRenderTargetCameraOutputInfo* ExistingRenderTargetOutputInfo =
+				MediaCaptureSettings->RenderTargetCaptures.FindByPredicate(GetFindOutputInfoPredicate<FMediaFrameworkCaptureRenderTargetCameraOutputInfo>(MediaOutput)))
+			{
+				MediaCaptureSettings->CurrentViewportMediaOutput.CaptureOptions = ExistingRenderTargetOutputInfo->CaptureOptions;
+			}
+
+			MediaCaptureSettings->CurrentViewportMediaOutput.MediaOutput = MediaOutput;
+		}
+
+		OnMethodChanged.ExecuteIfBound();
+	}
+	
 	void OnViewportMethodSelected()
 	{
 		UMediaFrameworkWorldSettingsAssetUserData* MediaCaptureSettings = FMediaProfileEditor::FindOrAddMediaFrameworkCaptureSettings();
@@ -156,6 +206,11 @@ private:
 				NewViewportOutputInfo.CaptureOptions = ExistingRenderTargetOutputInfo->CaptureOptions;
 			}
 
+			if (MediaCaptureSettings->CurrentViewportMediaOutput.MediaOutput == MediaOutput)
+			{
+				MediaCaptureSettings->CurrentViewportMediaOutput.MediaOutput = nullptr;
+			}
+
 			MediaCaptureSettings->ViewportCaptures.Add(NewViewportOutputInfo);
 			MediaCaptureSettings->RenderTargetCaptures.RemoveAll(GetFindOutputInfoPredicate<FMediaFrameworkCaptureRenderTargetCameraOutputInfo>(MediaOutput));
 		}
@@ -201,6 +256,11 @@ private:
 				NewRenderTargetOutputInfo.CaptureOptions = ExistingViewportOutputInfo->CaptureOptions;
 			}
 
+			if (MediaCaptureSettings->CurrentViewportMediaOutput.MediaOutput == MediaOutput)
+			{
+				MediaCaptureSettings->CurrentViewportMediaOutput.MediaOutput = nullptr;
+			}
+
 			MediaCaptureSettings->RenderTargetCaptures.Add(NewRenderTargetOutputInfo);
 			MediaCaptureSettings->ViewportCaptures.RemoveAll(GetFindOutputInfoPredicate<FMediaFrameworkCaptureCameraViewportCameraOutputInfo>(MediaOutput));
 		}
@@ -216,6 +276,7 @@ private:
 			return LOCTEXT("MediaCaptureMethodNotConfigured", "Not Configured");
 		}
 		
+		bool bCurrentViewport = false;
 		int32 NumViewports = 0;
 		int32 NumRenderTargets = 0;
 
@@ -228,17 +289,23 @@ private:
 				continue;
 			}
 
+			bCurrentViewport = MediaCaptureSettings->CurrentViewportMediaOutput.MediaOutput == MediaOutput;
 			NumViewports += Algo::CountIf(MediaCaptureSettings->ViewportCaptures, GetFindOutputInfoPredicate<FMediaFrameworkCaptureCameraViewportCameraOutputInfo>(MediaOutput));
 			NumRenderTargets += Algo::CountIf(MediaCaptureSettings->RenderTargetCaptures, GetFindOutputInfoPredicate<FMediaFrameworkCaptureRenderTargetCameraOutputInfo>(MediaOutput));
 		}
 
-		if (NumViewports > 0 && NumRenderTargets > 0)
+		const bool bMultipleTypes = (bCurrentViewport && (NumViewports > 0 || NumRenderTargets > 0)) || (NumViewports > 0 && NumRenderTargets > 0);
+		if (bMultipleTypes)
 		{
 			return LOCTEXT("MediaCaptureMethodMultipleValues", "Multiple Values");
 		}
+		if (bCurrentViewport)
+		{
+			return LOCTEXT("MediaCaptureCurrentViewportMethod", "Current Viewport");
+		}
 		if (NumViewports > 0)
 		{
-			return LOCTEXT("MediaCaptureViewportMethod", "Viewport");
+			return LOCTEXT("MediaCaptureViewportMethod", "Media Viewport");
 		}
 		if (NumRenderTargets > 0)
 		{
@@ -262,6 +329,34 @@ private:
 	FSimpleDelegate OnMethodChanged;
 };
 
+class FMediaFrameworkCurrentViewportOutputInfoCustomization : public IPropertyTypeCustomization
+{
+protected:
+	virtual void CustomizeHeader(TSharedRef<IPropertyHandle> PropertyHandle, FDetailWidgetRow& HeaderRow, IPropertyTypeCustomizationUtils& CustomizationUtils) override { }
+
+	virtual void CustomizeChildren(TSharedRef<IPropertyHandle> PropertyHandle, IDetailChildrenBuilder& ChildBuilder, IPropertyTypeCustomizationUtils& CustomizationUtils) override
+	{
+		// Manually add all capture options properties to avoid an additional property group getting added
+		TSharedPtr<IPropertyHandle> CaptureOptionsHandle = PropertyHandle->GetChildHandle(GET_MEMBER_NAME_CHECKED(FMediaFrameworkCaptureCurrentViewportOutputInfo, CaptureOptions));
+		uint32 NumChildren;
+		CaptureOptionsHandle->GetNumChildren(NumChildren);
+		for (uint32 Index = 0; Index < NumChildren; ++Index)
+		{
+			TSharedPtr<IPropertyHandle> ChildHandle = CaptureOptionsHandle->GetChildHandle(Index);
+			ChildHandle->SetOnPropertyValueChangedWithData(TDelegate<void(const FPropertyChangedEvent&)>::CreateSP(this, &FMediaFrameworkCurrentViewportOutputInfoCustomization::PropertyValueChange));			
+			ChildBuilder.AddProperty(ChildHandle.ToSharedRef());
+		}
+	}
+	
+	void PropertyValueChange(const FPropertyChangedEvent& PropertyChangedEvent)
+	{
+		if (UMediaFrameworkWorldSettingsAssetUserData* MediaCaptureSettings = FMediaProfileEditor::FindOrAddMediaFrameworkCaptureSettings())
+		{
+			FCoreUObjectDelegates::OnObjectPropertyChanged.Broadcast(MediaCaptureSettings, const_cast<FPropertyChangedEvent&>(PropertyChangedEvent));
+		}
+	}
+};
+
 class FMediaFrameworkViewportOutputInfoCustomization : public IPropertyTypeCustomization
 {
 protected:
@@ -329,6 +424,7 @@ protected:
 
 FMediaProfileMediaOutputDetailCustomization::~FMediaProfileMediaOutputDetailCustomization()
 {
+	FCoreUObjectDelegates::OnObjectPropertyChanged.RemoveAll(this);
 	if (MediaProfileEditor.IsValid())
 	{
 		MediaProfileEditor.Pin()->GetOnCaptureSettingsObjectChanged().RemoveAll(this);
@@ -345,9 +441,17 @@ void FMediaProfileMediaOutputDetailCustomization::CustomizeDetails(IDetailLayout
 	DetailBuilder.SortCategories([](const TMap<FName, IDetailCategoryBuilder*>& CategoryMap) { });
 	
 	MediaProfileEditor.Pin()->GetOnCaptureSettingsObjectChanged().AddSP(this, &FMediaProfileMediaOutputDetailCustomization::OnCaptureSettingsObjectChanged);
+	FCoreUObjectDelegates::OnObjectPropertyChanged.AddSP(this, &FMediaProfileMediaOutputDetailCustomization::OnObjectPropertyChanged);
 	
 	TArray<TWeakObjectPtr<UObject>> Objects;
 	CachedDetailBuilder.Pin()->GetObjectsBeingCustomized(Objects);
+
+	DetailBuilder.RegisterInstancedCustomPropertyTypeLayout(
+		FMediaFrameworkCaptureCurrentViewportOutputInfo::StaticStruct()->GetFName(),
+		FOnGetPropertyTypeCustomizationInstance::CreateLambda([]
+		{
+			return MakeShared<FMediaFrameworkCurrentViewportOutputInfoCustomization>();
+		}));
 	
 	DetailBuilder.RegisterInstancedCustomPropertyTypeLayout(
 		FMediaFrameworkCaptureCameraViewportCameraOutputInfo::StaticStruct()->GetFName(),
@@ -429,6 +533,7 @@ void FMediaProfileMediaOutputDetailCustomization::CustomizeDetails(IDetailLayout
 		];
 
 	TArray<TSharedPtr<FStructOnScope>> CaptureSettingStructs;
+	bool bHasCurrentViewportCapture = false;
 	bool bHasViewportCaptures = false;
 	bool bHasRenderTargetCaptures = false;
 	if (UMediaFrameworkWorldSettingsAssetUserData* MediaCaptureSettings = FMediaProfileEditor::FindOrAddMediaFrameworkCaptureSettings())
@@ -447,6 +552,17 @@ void FMediaProfileMediaOutputDetailCustomization::CustomizeDetails(IDetailLayout
 				continue;
 			}
 
+			if (MediaCaptureSettings->CurrentViewportMediaOutput.MediaOutput == MediaOutput)
+			{
+				TSharedPtr<FStructOnScope> Struct = MakeShared<FStructOnScope>(
+					FMediaFrameworkCaptureCurrentViewportOutputInfo::StaticStruct(),
+					reinterpret_cast<uint8*>(&MediaCaptureSettings->CurrentViewportMediaOutput));
+				Struct->SetPackage(MediaCaptureSettings->GetPackage());
+					
+				CaptureSettingStructs.Add(Struct);
+				bHasCurrentViewportCapture = true;
+			}
+			
 			for (FMediaFrameworkCaptureCameraViewportCameraOutputInfo& ViewportOutputInfo : MediaCaptureSettings->ViewportCaptures)
 			{
 				if (ViewportOutputInfo.MediaOutput == MediaOutput)
@@ -473,8 +589,9 @@ void FMediaProfileMediaOutputDetailCustomization::CustomizeDetails(IDetailLayout
 		}
 	}
 
-	// Only show the capture settings if all objects have capture settings of the same type (viewport or render target)
-	if (CaptureSettingStructs.Num() && bHasViewportCaptures != bHasRenderTargetCaptures)
+	// Only show the capture settings if all objects have capture settings of the same type (current viewport, media viewport, or render target)
+	const bool bHasMultipleValues = (bHasCurrentViewportCapture && (bHasViewportCaptures || bHasRenderTargetCaptures)) || (bHasViewportCaptures && bHasRenderTargetCaptures);
+	if (CaptureSettingStructs.Num() && !bHasMultipleValues)
 	{
 		FAddPropertyParams AddPropertyParams;
 		AddPropertyParams.HideRootObjectNode(true);
@@ -565,6 +682,12 @@ void FMediaProfileMediaOutputDetailCustomization::SetMediaObject(UMediaOutput* I
 	}
 
 	bool bHasExistingCaptureInfo = false;
+	if (CaptureSettings->CurrentViewportMediaOutput.MediaOutput == InOriginalMediaObject)
+	{
+		CaptureSettings->CurrentViewportMediaOutput.MediaOutput = InNewMediaObject;
+		bHasExistingCaptureInfo = true;
+	}
+	
 	for (FMediaFrameworkCaptureCameraViewportCameraOutputInfo& ViewportOutputInfo : CaptureSettings->ViewportCaptures)
 	{
 		if (ViewportOutputInfo.MediaOutput != InOriginalMediaObject)
@@ -792,4 +915,26 @@ void FMediaProfileMediaOutputDetailCustomization::OnCaptureMethodChanged()
 	CachedDetailBuilder.Pin()->ForceRefreshDetails();
 }
 
+void FMediaProfileMediaOutputDetailCustomization::OnObjectPropertyChanged(UObject* InObject, FPropertyChangedEvent& InPropertyChangedEvent)
+{
+	if (UMediaFrameworkWorldSettingsAssetUserData* CaptureSettings = Cast<UMediaFrameworkWorldSettingsAssetUserData>(InObject))
+	{
+		if (CaptureSettings != FMediaProfileEditor::FindOrAddMediaFrameworkCaptureSettings())
+		{
+			return;
+		}
+		
+		const FName PropertyName = InPropertyChangedEvent.GetPropertyName();
+		
+		// Refresh the details panel if the MediaOutput property on any captures or the viewport or render target captures lists have changed,
+		// as the capture properties may now be invalid
+		if (PropertyName == TEXT("MediaOutput") ||
+			PropertyName == GET_MEMBER_NAME_CHECKED(UMediaFrameworkWorldSettingsAssetUserData, ViewportCaptures) ||
+			PropertyName == GET_MEMBER_NAME_CHECKED(UMediaFrameworkWorldSettingsAssetUserData, RenderTargetCaptures))
+		{
+			CachedDetailBuilder.Pin()->ForceRefreshDetails();
+		}
+	}
+}
+
 #undef LOCTEXT_NAMESPACE
diff --git a/Engine/Plugins/Media/MediaFrameworkUtilities/Source/MediaFrameworkUtilitiesEditor/Private/AssetEditor/MediaProfileMediaItemDetailCustomization.h b/Engine/Plugins/Media/MediaFrameworkUtilities/Source/MediaFrameworkUtilitiesEditor/Private/AssetEditor/MediaProfileMediaItemDetailCustomization.h
--- a/Engine/Plugins/Media/MediaFrameworkUtilities/Source/MediaFrameworkUtilitiesEditor/Private/AssetEditor/MediaProfileMediaItemDetailCustomization.h
+++ b/Engine/Plugins/Media/MediaFrameworkUtilities/Source/MediaFrameworkUtilitiesEditor/Private/AssetEditor/MediaProfileMediaItemDetailCustomization.h
@@ -377,6 +377,7 @@ private:
 	void OnCaptureSettingsObjectChanged();
 
 	void OnCaptureMethodChanged();
+	void OnObjectPropertyChanged(UObject* InObject, FPropertyChangedEvent& InPropertyChangedEvent);
 	
 private:
 	/**
diff --git a/Engine/Plugins/Media/MediaFrameworkUtilities/Source/MediaFrameworkUtilitiesEditor/Private/AssetEditor/SMediaProfileDetailsInfoPanel.cpp b/Engine/Plugins/Media/MediaFrameworkUtilities/Source/MediaFrameworkUtilitiesEditor/Private/AssetEditor/SMediaProfileDetailsInfoPanel.cpp
--- a/Engine/Plugins/Media/MediaFrameworkUtilities/Source/MediaFrameworkUtilitiesEditor/Private/AssetEditor/SMediaProfileDetailsInfoPanel.cpp
+++ b/Engine/Plugins/Media/MediaFrameworkUtilities/Source/MediaFrameworkUtilitiesEditor/Private/AssetEditor/SMediaProfileDetailsInfoPanel.cpp
@@ -323,24 +323,31 @@ FText SMediaObjectInfoPanel::GetCaptureMethodText() const
 		return FText::GetEmpty();
 	}
 
-	bool bHasViewportCaptures = CaptureSettings->ViewportCaptures.ContainsByPredicate([MediaOutput](const FMediaFrameworkCaptureCameraViewportCameraOutputInfo& Info)
+	const bool bHasCurrentViewportCapture = CaptureSettings->CurrentViewportMediaOutput.MediaOutput == MediaOutput;
+	const bool bHasViewportCaptures = CaptureSettings->ViewportCaptures.ContainsByPredicate([MediaOutput](const FMediaFrameworkCaptureCameraViewportCameraOutputInfo& Info)
 	{
 		return Info.MediaOutput == MediaOutput;
 	});
 
-	bool bHasRenderTargetCaptures = CaptureSettings->RenderTargetCaptures.ContainsByPredicate([MediaOutput](const FMediaFrameworkCaptureRenderTargetCameraOutputInfo& Info)
+	const bool bHasRenderTargetCaptures = CaptureSettings->RenderTargetCaptures.ContainsByPredicate([MediaOutput](const FMediaFrameworkCaptureRenderTargetCameraOutputInfo& Info)
 	{
 		return Info.MediaOutput == MediaOutput;
 	});
 
-	if (bHasViewportCaptures && bHasRenderTargetCaptures)
+	const bool bHasMultipleValues = (bHasCurrentViewportCapture && (bHasViewportCaptures || bHasRenderTargetCaptures)) || (bHasViewportCaptures && bHasRenderTargetCaptures);
+	if (bHasMultipleValues)
 	{
 		return LOCTEXT("MultipleCapturesValue", "Multiple Captures");
 	}
 
+	if (bHasCurrentViewportCapture)
+	{
+		return LOCTEXT("CurrentViewportCaptureValue", "Current Viewport");	
+	}
+
 	if (bHasViewportCaptures)
 	{
-		return LOCTEXT("ViewportCaptureValue", "Viewport");
+		return LOCTEXT("ViewportCaptureValue", "Media Viewport");
 	}
 
 	if (bHasRenderTargetCaptures)
@@ -360,6 +367,7 @@ FText SMediaObjectInfoPanel::GetCaptureObjectLabelText() const
 		return FText::GetEmpty();
 	}
 
+	const bool bHasCurrentViewportCapture = CaptureSettings->CurrentViewportMediaOutput.MediaOutput == MediaOutput;
 	bool bHasViewportCaptures = CaptureSettings->ViewportCaptures.ContainsByPredicate([MediaOutput](const FMediaFrameworkCaptureCameraViewportCameraOutputInfo& Info)
 	{
 		return Info.MediaOutput == MediaOutput;
@@ -370,11 +378,17 @@ FText SMediaObjectInfoPanel::GetCaptureObjectLabelText() const
 		return Info.MediaOutput == MediaOutput;
 	});
 
-	if (bHasViewportCaptures && bHasRenderTargetCaptures)
+	const bool bHasMultipleValues = (bHasCurrentViewportCapture && (bHasViewportCaptures || bHasRenderTargetCaptures)) || (bHasViewportCaptures && bHasRenderTargetCaptures);
+	if (bHasMultipleValues)
 	{
 		return LOCTEXT("MultipleCaptureObjectsLabel", "Multiple Objects:");
 	}
 
+	if (bHasCurrentViewportCapture)
+	{
+		return FText::GetEmpty();
+	}
+
 	if (bHasViewportCaptures)
 	{
 		return LOCTEXT("CameraObjectLabel", "Cameras:");
@@ -397,6 +411,8 @@ FText SMediaObjectInfoPanel::GetCaptureObjectValueText() const
 		return FText::GetEmpty();
 	}
 
+	const bool bHasCurrentViewportCapture = CaptureSettings->CurrentViewportMediaOutput.MediaOutput == MediaOutput;
+
 	bool bHasViewportCaptures = false;
 	FText CamerasText;
 	for (const FMediaFrameworkCaptureCameraViewportCameraOutputInfo& Info : CaptureSettings->ViewportCaptures)
@@ -449,7 +465,8 @@ FText SMediaObjectInfoPanel::GetCaptureObjectValueText() const
 		}
 	}
 
-	if (bHasViewportCaptures && bHasRenderTargetCaptures)
+	const bool bHasMultipleValues = (bHasCurrentViewportCapture && (bHasViewportCaptures || bHasRenderTargetCaptures)) || (bHasViewportCaptures && bHasRenderTargetCaptures);
+	if (bHasMultipleValues)
 	{
 		return FText::GetEmpty();
 	}
@@ -499,6 +516,21 @@ FText SMediaObjectInfoPanel::GetCaptureColorConversionText() const
 	}
 
 	FText ColorConversionText;
+	if (CaptureSettings->CurrentViewportMediaOutput.MediaOutput == MediaOutput)
+	{
+		if (ColorConversionText.IsEmpty())
+		{
+			ColorConversionText = FText::FromString(CaptureSettings->CurrentViewportMediaOutput.CaptureOptions.GetColorConversionSettingsString());
+		}
+		else
+		{
+			ColorConversionText = FText::Format(
+				LOCTEXT("AppendColorConversionFormat", "{0}, {1}"),
+				ColorConversionText,
+				FText::FromString(CaptureSettings->CurrentViewportMediaOutput.CaptureOptions.GetColorConversionSettingsString()));
+		}
+	}
+	
 	for (const FMediaFrameworkCaptureCameraViewportCameraOutputInfo& Info : CaptureSettings->ViewportCaptures)
 	{
 		if (Info.MediaOutput != MediaOutput)
diff --git a/Engine/Plugins/Media/MediaFrameworkUtilities/Source/MediaFrameworkUtilitiesEditor/Private/AssetEditor/SMediaProfileMediaItemDisplay.cpp b/Engine/Plugins/Media/MediaFrameworkUtilities/Source/MediaFrameworkUtilitiesEditor/Private/AssetEditor/SMediaProfileMediaItemDisplay.cpp
--- a/Engine/Plugins/Media/MediaFrameworkUtilities/Source/MediaFrameworkUtilitiesEditor/Private/AssetEditor/SMediaProfileMediaItemDisplay.cpp
+++ b/Engine/Plugins/Media/MediaFrameworkUtilities/Source/MediaFrameworkUtilitiesEditor/Private/AssetEditor/SMediaProfileMediaItemDisplay.cpp
@@ -202,12 +202,262 @@ FText SMediaProfileMediaSourceDisplay::GetFramerateText() const
 	return FText::Format(LOCTEXT("MediaSourceFramerateFormat", "{0} fps"), MediaPlayer->GetVideoTrackFrameRate(INDEX_NONE, INDEX_NONE));
 }
 
+/** Base widget for rendering the output of a media capture */
+class SMediaCaptureImageBase : public SCompoundWidget
+{
+public:
+	SLATE_BEGIN_ARGS(SMediaCaptureImageBase) {}
+		SLATE_ATTRIBUTE(TOptional<EMediaCaptureState>, MediaCaptureState)
+		SLATE_ATTRIBUTE(FOptionalSize, AspectRatio)
+	SLATE_END_ARGS()
+
+	void Construct(const FArguments& InArgs)
+	{
+		MediaCaptureState = InArgs._MediaCaptureState;
+
+		ChildSlot
+		[
+			SNew(SBox)
+			.MinAspectRatio(InArgs._AspectRatio)
+			.MaxAspectRatio(InArgs._AspectRatio)
+			[
+				SNew(SOverlay)
+
+				+SOverlay::Slot()
+				[
+					GetContentWidget()
+				]
+				
+				+SOverlay::Slot()
+				.HAlign(HAlign_Left)
+				.VAlign(VAlign_Top)
+				.Padding(16)
+				[
+					SNew(SHorizontalBox)
+
+					+SHorizontalBox::Slot()
+					.AutoWidth()
+					.HAlign(HAlign_Center)
+					.VAlign(VAlign_Center)
+					.Padding(0.0f, 0.0f, 8.0f, 0.0f)
+					[
+						SNew(SImage).Image(this, &SMediaCaptureImageBase::GetLiveCaptureIcon)
+						.ColorAndOpacity(this, &SMediaCaptureImageBase::GetLiveCaptureIconColor)
+					]
+
+					+SHorizontalBox::Slot()
+					.AutoWidth()
+					.HAlign(HAlign_Center)
+					.VAlign(VAlign_Center)
+					[
+						SNew(STextBlock).Text(this, &SMediaCaptureImageBase::GetLiveCaptureText)
+					]
+				]
+			]
+		];
+	}
+	
+protected:
+	virtual TSharedRef<SWidget> GetContentWidget() { return SNullWidget::NullWidget; }
+	
+	const FSlateBrush* GetLiveCaptureIcon() const
+	{
+		TOptional<EMediaCaptureState> State = MediaCaptureState.Get();
+		switch (State.Get(EMediaCaptureState::Stopped))
+		{
+		case EMediaCaptureState::Capturing:
+			return FMediaFrameworkUtilitiesEditorStyle::Get().GetBrush("MediaCapture.Capture");
+
+		case EMediaCaptureState::Error:
+			return FAppStyle::Get().GetBrush("Icons.XCircle");
+
+		case EMediaCaptureState::Preparing:
+		case EMediaCaptureState::StopRequested:
+			return FAppStyle::Get().GetBrush("Icons.AlertCircle");
+			
+		case EMediaCaptureState::Stopped:
+		default:
+			return FAppStyle::GetBrush("Icons.Toolbar.Stop");
+		}
+	}
+
+	FSlateColor GetLiveCaptureIconColor() const
+	{
+		TOptional<EMediaCaptureState> State = MediaCaptureState.Get();
+		switch (State.Get(EMediaCaptureState::Stopped))
+		{
+		case EMediaCaptureState::Capturing:
+		case EMediaCaptureState::Preparing:
+		case EMediaCaptureState::Stopped:
+		default:
+			return FSlateColor::UseForeground();
+
+		case EMediaCaptureState::Error:
+			return FStyleColors::Error;
+
+		case EMediaCaptureState::StopRequested:
+			return FStyleColors::Warning;
+		}
+	}
+	
+	FText GetLiveCaptureText() const
+	{
+		TOptional<EMediaCaptureState> State = MediaCaptureState.Get();
+		switch (State.Get(EMediaCaptureState::Stopped))
+		{
+		case EMediaCaptureState::Capturing:
+			return LOCTEXT("CapturingLabel", "Capturing");
+
+		case EMediaCaptureState::Error:
+			return LOCTEXT("ErrorCapturingLabel", "Error occured while capturing");
+			
+		case EMediaCaptureState::StopRequested:
+			return LOCTEXT("StoppingCaptureLabel", "Stopping capture...");
+
+		case EMediaCaptureState::Preparing:
+			return LOCTEXT("PreparingCaptureLabel", "Preparing capture...");
+			
+		case EMediaCaptureState::Stopped:
+		default:
+			return LOCTEXT("NotCapturingLabel", "Not Capturing");
+		}
+	}
+	
+protected:
+	/** Attribute for querying the current media capture state of the media output being displayed */
+	TAttribute<TOptional<EMediaCaptureState>> MediaCaptureState;
+};
+
+/** Widget that outputs the contents of a scene viewport render target */
+class SActiveEditorViewportRenderer : public SCompoundWidget
+{
+public:
+	SLATE_BEGIN_ARGS(SActiveEditorViewportRenderer) {}
+		SLATE_ATTRIBUTE(TWeakPtr<FSceneViewport>, ActiveEditorViewport)
+	SLATE_END_ARGS()
+	
+	void Construct(const FArguments& InArgs)
+	{
+		ActiveEditorViewport = InArgs._ActiveEditorViewport;
+
+		ChildSlot
+		[
+			SNew(SBorder)
+			.HAlign(HAlign_Center)
+			.VAlign(VAlign_Center)
+			.BorderImage(FAppStyle::GetBrush("NoBorder"))
+			[
+				SNew(STextBlock)
+				.Text(LOCTEXT("NoActiveEditorViewportLabel", "No active editor viewports found"))
+			]
+		];
+	}
+	
+	virtual int32 OnPaint(const FPaintArgs& Args, const FGeometry& AllottedGeometry, const FSlateRect& MyCullingRect, FSlateWindowElementList& OutDrawElements, int32 LayerId, const FWidgetStyle& InWidgetStyle, bool bParentEnabled) const override
+	{
+		TSharedPtr<FSceneViewport> Viewport = ActiveEditorViewport.IsSet() ? ActiveEditorViewport.Get().Pin() : nullptr;
+		if (!Viewport.IsValid() || Viewport->GetViewportRenderTargetTexture() == nullptr)
+		{
+			// If the viewport is not valid, draw the widget's children, which will display an error message
+			return SCompoundWidget::OnPaint(Args, AllottedGeometry, MyCullingRect, OutDrawElements, LayerId, InWidgetStyle, bParentEnabled);
+		}
+		
+		ESlateDrawEffect DrawEffects = ESlateDrawEffect::None;
+		DrawEffects |= ESlateDrawEffect::IgnoreTextureAlpha;
+		DrawEffects |= ESlateDrawEffect::NoGamma;
+		DrawEffects |= ESlateDrawEffect::NoBlending;
+		
+		FSlateDrawElement::MakeViewport(OutDrawElements, LayerId, AllottedGeometry.ToPaintGeometry(), Viewport, DrawEffects);
+		
+		return LayerId + 1;
+	}
+
+private:
+	/**
+	 * Attribute to retrieve the active editor viewport used for the media capture. Using weak pointer to avoid TAttribute caching a shared pointer,
+	 * as PIE viewports specifically don't like having stray shared pointers floating around
+	 */
+	TAttribute<TWeakPtr<FSceneViewport>> ActiveEditorViewport;
+};
+
+/** Widget that displays a current viewport media capture */
+class SMediaCaptureEditorViewportImage : public SMediaCaptureImageBase
+{
+public:
+	SLATE_BEGIN_ARGS(SMediaCaptureEditorViewportImage) {}
+		SLATE_ATTRIBUTE(TOptional<EMediaCaptureState>, MediaCaptureState)
+		SLATE_ATTRIBUTE(FOptionalSize, AspectRatio)
+	SLATE_END_ARGS()
+
+	void Construct(const FArguments& InArgs, const TWeakPtr<FMediaProfileEditor>& InMediaProfileEditor)
+	{
+		MediaProfileEditor = InMediaProfileEditor;
+		AspectRatio = InArgs._AspectRatio;
+		
+		SMediaCaptureImageBase::FArguments BaseArgs;
+		BaseArgs.MediaCaptureState(InArgs._MediaCaptureState);
+		BaseArgs.AspectRatio(this, &SMediaCaptureEditorViewportImage::GetAspectRatio);
+		
+		SMediaCaptureImageBase::Construct(BaseArgs);
+	}
+
+protected:
+	virtual TSharedRef<SWidget> GetContentWidget() override
+	{
+		return SNew(SActiveEditorViewportRenderer)
+			.ActiveEditorViewport(this, &SMediaCaptureEditorViewportImage::GetActiveEditorViewport);
+	}
+
+private:
+	TWeakPtr<FSceneViewport> GetActiveEditorViewport() const
+	{
+		if (TSharedPtr<FMediaProfileEditor> PinnedMediaProfileEditor = MediaProfileEditor.Pin())
+		{
+			return PinnedMediaProfileEditor->GetActiveEditorViewport();
+		}
+
+		return nullptr;
+	}
+
+	FOptionalSize GetAspectRatio() const
+	{
+		// If we are actively capturing the active editor viewport, use the supplied aspect ratio
+		// Otherwise, use the viewport's desired aspect ratio
+		const EMediaCaptureState State = MediaCaptureState.IsSet() ? MediaCaptureState.Get().Get(EMediaCaptureState::Stopped) : EMediaCaptureState::Stopped;
+		if (State == EMediaCaptureState::Capturing)
+		{
+			return AspectRatio.Get(FOptionalSize());
+		}
+
+		if (TSharedPtr<FMediaProfileEditor> PinnedMediaProfileEditor = MediaProfileEditor.Pin())
+		{
+			if (TSharedPtr<FSceneViewport> Viewport = PinnedMediaProfileEditor->GetActiveEditorViewport())
+			{
+				return Viewport->GetDesiredAspectRatio();
+			}
+		}
+
+		return FOptionalSize();
+	}
+	
+private:
+	/** The media profile editor that manages the media output being displayed */
+	TWeakPtr<FMediaProfileEditor> MediaProfileEditor;
+
+	/** Desired aspect ratio of the media output being captured */
+	TAttribute<FOptionalSize> AspectRatio;
+	
+	/** Command list that manages the show flags for the viewport capture */
+	TSharedPtr<FUICommandList> ShowFlagsCommandList;
+};
+
 /** Widget to display a viewport media capture output */
-class SMediaCaptureViewportImage : public SCompoundWidget
+class SMediaCaptureViewportImage : public SMediaCaptureImageBase
 {
 public:
 	SLATE_BEGIN_ARGS(SMediaCaptureViewportImage) {}
-		SLATE_ATTRIBUTE(bool, IsCapturing)
+		SLATE_ATTRIBUTE(TOptional<EMediaCaptureState>, MediaCaptureState)
+		SLATE_ATTRIBUTE(FOptionalSize, AspectRatio)
 	SLATE_END_ARGS()
 
 	virtual ~SMediaCaptureViewportImage() override
@@ -218,63 +468,16 @@ public:
 	
 	void Construct(const FArguments& InArgs, const TSharedPtr<FMediaProfileEditor::FViewportCapture> InViewportCapture, const TArray<TLazyObjectPtr<AActor>>& InLockedActors)
 	{
-		IsCapturing = InArgs._IsCapturing;
 		ViewportCapture = InViewportCapture;
 
 		FEditorDelegates::PostPIEStarted.AddSP(this, &SMediaCaptureViewportImage::OnPostPIEStarted);
 		FEditorDelegates::PrePIEEnded.AddSP(this, &SMediaCaptureViewportImage::OnPrePIEEnded);
 		
-		TSharedPtr<SWidget> Content;
-		if (InViewportCapture.IsValid() && InViewportCapture->ViewportWidget.IsValid())
-		{
-			Content = InViewportCapture->ViewportWidget;
-		}
-		else
-		{
-			Content = SNew(SBorder)
-				.HAlign(HAlign_Center)
-				.VAlign(VAlign_Center)
-				.BorderImage(FAppStyle::GetBrush("NoBorder"))
-				[
-					SNew(STextBlock)
-					.Text(LOCTEXT("ViewportNotConfiguredLabel", "Viewport output not configured for Media Output Capture"))
-				];
-		}
+		SMediaCaptureImageBase::FArguments BaseArgs;
+		BaseArgs.MediaCaptureState(InArgs._MediaCaptureState);
+		BaseArgs.AspectRatio(InArgs._AspectRatio);
 		
-		ChildSlot
-		[
-			SNew(SOverlay)
-
-			+SOverlay::Slot()
-			[
-				Content.ToSharedRef()
-			]
-			
-			+SOverlay::Slot()
-			.HAlign(HAlign_Left)
-			.VAlign(VAlign_Top)
-			.Padding(16)
-			[
-				SNew(SHorizontalBox)
-
-				+SHorizontalBox::Slot()
-				.AutoWidth()
-				.HAlign(HAlign_Center)
-				.VAlign(VAlign_Center)
-				.Padding(0.0f, 0.0f, 8.0f, 0.0f)
-				[
-					SNew(SImage).Image(this, &SMediaCaptureViewportImage::GetLiveCaptureIcon)
-				]
-
-				+SHorizontalBox::Slot()
-				.AutoWidth()
-				.HAlign(HAlign_Center)
-				.VAlign(VAlign_Center)
-				[
-					SNew(STextBlock).Text(this, &SMediaCaptureViewportImage::GetLiveCaptureText)
-				]
-			]
-		];
+		SMediaCaptureImageBase::Construct(BaseArgs);
 
 		Refresh(InLockedActors);
 	}
@@ -348,20 +551,26 @@ public:
 	{
 		PendingViewportSizeChange = InDesiredViewportSize;
 	}
-	
-private:
-	const FSlateBrush* GetLiveCaptureIcon() const
+
+protected:
+	virtual TSharedRef<SWidget> GetContentWidget() override
 	{
-		bool bIsCapturing = IsCapturing.Get(false);
-		return bIsCapturing ? FMediaFrameworkUtilitiesEditorStyle::Get().GetBrush("MediaCapture.Capture") : FAppStyle::GetBrush("Icons.Toolbar.Stop");
+		if (ViewportCapture.IsValid() && ViewportCapture->ViewportWidget.IsValid())
+		{
+			return ViewportCapture->ViewportWidget.ToSharedRef();
+		}
+		
+		return SNew(SBorder)
+			.HAlign(HAlign_Center)
+			.VAlign(VAlign_Center)
+			.BorderImage(FAppStyle::GetBrush("NoBorder"))
+			[
+				SNew(STextBlock)
+				.Text(LOCTEXT("ViewportNotConfiguredLabel", "Viewport output not configured for Media Output Capture"))
+			];
 	}
 	
-	FText GetLiveCaptureText() const
-	{
-		bool bIsCapturing = IsCapturing.Get(false);
-		return bIsCapturing ? LOCTEXT("CapturingLabel", "Capturing") : LOCTEXT("NotCapturingLabel", "Not Capturing");
-	}
-
+private:
 	/** Raised when Play-in-Editor is started */
 	void OnPostPIEStarted(bool bIsSimulating)
 	{
@@ -387,9 +596,6 @@ private:
 	/** The current locked actor that functions as the camera. */
 	int32 CurrentLockedActor = INDEX_NONE;
 
-	/** Attribute for querying if the viewport capture is currently actively being captured */
-	TAttribute<bool> IsCapturing;
-
 	/**
 	 * To manually set the viewport size on the capture, we need the SViewport widget to be fully in the slate hierarchy
 	 * (as FSceneViewport needs to be able to find an SWindow from SViewport), so this optional stores any pending size change
@@ -399,51 +605,21 @@ private:
 };
 
 /** Widget to display a render target media capture output */
-class SMediaCaptureRenderTargetImage : public SCompoundWidget
+class SMediaCaptureRenderTargetImage : public SMediaCaptureImageBase
 {
 public:
 	SLATE_BEGIN_ARGS(SMediaCaptureRenderTargetImage) {}
-		SLATE_ATTRIBUTE(bool, IsCapturing)
+		SLATE_ATTRIBUTE(TOptional<EMediaCaptureState>, MediaCaptureState)
+		SLATE_ATTRIBUTE(FOptionalSize, AspectRatio)
 	SLATE_END_ARGS()
 
 	void Construct(const FArguments& InArgs, UTextureRenderTarget2D* InRenderTarget)
 	{
-		IsCapturing = InArgs._IsCapturing;
+		SMediaCaptureImageBase::FArguments BaseArgs;
+		BaseArgs.MediaCaptureState(InArgs._MediaCaptureState);
+		BaseArgs.AspectRatio(InArgs._AspectRatio);
 		
-		ChildSlot
-		[
-			SNew(SOverlay)
-
-			+SOverlay::Slot()
-			[
-				SAssignNew(Container, SBox)
-			]
-
-			+SOverlay::Slot()
-			.HAlign(HAlign_Left)
-			.VAlign(VAlign_Top)
-			.Padding(16)
-			[
-				SNew(SHorizontalBox)
-
-				+SHorizontalBox::Slot()
-				.AutoWidth()
-				.HAlign(HAlign_Center)
-				.VAlign(VAlign_Center)
-				.Padding(0.0f, 0.0f, 8.0f, 0.0f)
-				[
-					SNew(SImage).Image(this, &SMediaCaptureRenderTargetImage::GetLiveCaptureIcon)
-				]
-
-				+SHorizontalBox::Slot()
-				.AutoWidth()
-				.HAlign(HAlign_Center)
-				.VAlign(VAlign_Center)
-				[
-					SNew(STextBlock).Text(this, &SMediaCaptureRenderTargetImage::GetLiveCaptureText)
-				]
-			]
-		];
+		SMediaCaptureImageBase::Construct(BaseArgs);
 
 		Refresh(InRenderTarget);
 	}
@@ -480,17 +656,10 @@ public:
 		}
 	}
 
-private:
-	const FSlateBrush* GetLiveCaptureIcon() const
-	{
-		bool bIsCapturing = IsCapturing.Get(false);
-		return bIsCapturing ? FMediaFrameworkUtilitiesEditorStyle::Get().GetBrush("MediaCapture.Capture") : FAppStyle::GetBrush("Icons.Toolbar.Stop");
-	}
-	
-	FText GetLiveCaptureText() const
+protected:
+	virtual TSharedRef<SWidget> GetContentWidget() override
 	{
-		bool bIsCapturing = IsCapturing.Get(false);
-		return bIsCapturing ? LOCTEXT("CapturingLabel", "Capturing") : LOCTEXT("NotCapturingLabel", "Not Capturing");
+		return SAssignNew(Container, SBox);
 	}
 	
 private:
@@ -557,6 +726,16 @@ void SMediaProfileMediaOutputDisplay::ConfigureMediaImage()
 	
 	if (UMediaFrameworkWorldSettingsAssetUserData* CaptureSettings = FMediaProfileEditor::FindOrAddMediaFrameworkCaptureSettings())
 	{
+		if (CaptureSettings->CurrentViewportMediaOutput.MediaOutput == MediaOutput)
+		{
+			MediaImageContainer->SetContent(
+				SAssignNew(ImageWidget, SMediaCaptureEditorViewportImage, MediaProfileEditor)
+				.MediaCaptureState(this, &SMediaProfileMediaOutputDisplay::GetMediaOutputCaptureState)
+				.AspectRatio(this, &SMediaProfileMediaOutputDisplay::GetMediaOutputDesiredAspectRatio)
+			);
+			return;
+		}
+		
 		ViewportOutputInfoIndex = CaptureSettings->ViewportCaptures.IndexOfByPredicate(
 		[MediaOutput](const FMediaFrameworkCaptureCameraViewportCameraOutputInfo& OutputInfo)
 			{
@@ -568,25 +747,19 @@ void SMediaProfileMediaOutputDisplay::ConfigureMediaImage()
 			TSharedPtr<FMediaProfileEditor::FViewportCapture> ViewportCapture = MediaProfileEditor.IsValid() ? MediaProfileEditor.Pin()->FindOrCreateViewportCapture(GetMediaItem()) : nullptr;
 			if (ViewportCapture.IsValid())
 			{
+				TSharedPtr<SMediaCaptureViewportImage> ViewportImage;
+				MediaImageContainer->SetContent(
+					SAssignNew(ViewportImage, SMediaCaptureViewportImage, ViewportCapture, CaptureSettings->ViewportCaptures[ViewportOutputInfoIndex].LockedActors)
+					.MediaCaptureState(this, &SMediaProfileMediaOutputDisplay::GetMediaOutputCaptureState)
+					.AspectRatio(this, &SMediaProfileMediaOutputDisplay::GetMediaOutputDesiredAspectRatio)
+				);
+
 				FIntPoint TargetSize = MediaOutput->GetRequestedSize();
 				if (TargetSize == UMediaOutput::RequestCaptureSourceSize)
 				{
 					TargetSize = GetMutableDefault<UMediaFrameworkMediaCaptureSettings>()->DefaultViewportCaptureSize;
 				}
-
-				const float AspectRatio = (TargetSize.X > 0 && TargetSize.Y > 0) ? (float)TargetSize.X / (float)TargetSize.Y : 1.77777777f;
-
-				TSharedPtr<SMediaCaptureViewportImage> ViewportImage;
-				MediaImageContainer->SetContent(
-					SNew(SBox)
-					.MinAspectRatio(AspectRatio)
-					.MaxAspectRatio(AspectRatio)
-					[
-						SAssignNew(ViewportImage, SMediaCaptureViewportImage, ViewportCapture, CaptureSettings->ViewportCaptures[ViewportOutputInfoIndex].LockedActors)
-						.IsCapturing(this, &SMediaProfileMediaOutputDisplay::IsCapturingMediaOutput)
-					]
-				);
-
+				
 				ViewportImage->QueueViewportSizeChange(TargetSize);
 				ImageWidget = ViewportImage;
 				
@@ -604,21 +777,10 @@ void SMediaProfileMediaOutputDisplay::ConfigureMediaImage()
 
 		if (RenderTargetOutputInfoIndex != INDEX_NONE)
 		{
-			FIntPoint TargetSize = MediaOutput->GetRequestedSize();
-			if (TargetSize == UMediaOutput::RequestCaptureSourceSize)
-			{
-				TargetSize = GetMutableDefault<UMediaFrameworkMediaCaptureSettings>()->DefaultViewportCaptureSize;
-			}
-
-			const float AspectRatio = (TargetSize.X > 0 && TargetSize.Y > 0) ? (float)TargetSize.X / (float)TargetSize.Y : 1.77777777f;
 			MediaImageContainer->SetContent(
-				SNew(SBox)
-				.MinAspectRatio(AspectRatio)
-				.MaxAspectRatio(AspectRatio)
-				[
-					SAssignNew(ImageWidget, SMediaCaptureRenderTargetImage, CaptureSettings->RenderTargetCaptures[RenderTargetOutputInfoIndex].RenderTarget)
-					.IsCapturing(this, &SMediaProfileMediaOutputDisplay::IsCapturingMediaOutput)
-				]
+				SAssignNew(ImageWidget, SMediaCaptureRenderTargetImage, CaptureSettings->RenderTargetCaptures[RenderTargetOutputInfoIndex].RenderTarget)
+				.MediaCaptureState(this, &SMediaProfileMediaOutputDisplay::GetMediaOutputCaptureState)
+				.AspectRatio(this, &SMediaProfileMediaOutputDisplay::GetMediaOutputDesiredAspectRatio)
 			);
 			return;
 		}
@@ -687,8 +849,9 @@ void SMediaProfileMediaOutputDisplay::OnObjectPropertyChanged(UObject* InObject,
 		{
 			return;
 		}
-		
-		if (InPropertyChangedEvent.GetPropertyName() == GET_MEMBER_NAME_CHECKED(FMediaFrameworkCaptureRenderTargetCameraOutputInfo, RenderTarget))
+
+		const FName PropertyName = InPropertyChangedEvent.GetPropertyName();
+		if (PropertyName == GET_MEMBER_NAME_CHECKED(FMediaFrameworkCaptureRenderTargetCameraOutputInfo, RenderTarget))
 		{
 			if (RenderTargetOutputInfoIndex == INDEX_NONE || !CaptureSettings->RenderTargetCaptures.IsValidIndex(RenderTargetOutputInfoIndex))
 			{
@@ -705,7 +868,7 @@ void SMediaProfileMediaOutputDisplay::OnObjectPropertyChanged(UObject* InObject,
 			return;
 		}
 
-		if (InPropertyChangedEvent.GetPropertyName() == GET_MEMBER_NAME_CHECKED(FMediaFrameworkCaptureCameraViewportCameraOutputInfo, LockedActors))
+		if (PropertyName == GET_MEMBER_NAME_CHECKED(FMediaFrameworkCaptureCameraViewportCameraOutputInfo, LockedActors))
 		{
 			if (ViewportOutputInfoIndex == INDEX_NONE || !CaptureSettings->ViewportCaptures.IsValidIndex(ViewportOutputInfoIndex))
 			{
@@ -721,18 +884,68 @@ void SMediaProfileMediaOutputDisplay::OnObjectPropertyChanged(UObject* InObject,
 			ViewportImageWidget->Refresh(CaptureSettings->ViewportCaptures[ViewportOutputInfoIndex].LockedActors);
 			return;
 		}
+
+		// If a media output property in the capture settings has changed, or the viewport captures or render target captures lists have changed,
+		// check to see the media output object for this display is still linked to a valid capture. If not, refresh the widget
+		if (PropertyName == TEXT("MediaOutput") ||
+			PropertyName == GET_MEMBER_NAME_CHECKED(UMediaFrameworkWorldSettingsAssetUserData, ViewportCaptures) ||
+			PropertyName == GET_MEMBER_NAME_CHECKED(UMediaFrameworkWorldSettingsAssetUserData, RenderTargetCaptures))
+		{
+			const bool bIsViewportCapture = ViewportOutputInfoIndex != INDEX_NONE;
+			const bool bIsRenderTargetCapture = RenderTargetOutputInfoIndex != INDEX_NONE;
+			const bool bIsActiveEditorViewportCapture = ViewportOutputInfoIndex == INDEX_NONE && RenderTargetOutputInfoIndex == INDEX_NONE;
+
+			bool bResetMediaImage = false;
+			if (bIsViewportCapture)
+			{
+				bResetMediaImage = !CaptureSettings->ViewportCaptures.IsValidIndex(ViewportOutputInfoIndex) ||
+					CaptureSettings->ViewportCaptures[ViewportOutputInfoIndex].MediaOutput != GetMediaItem();
+			}
+			else if (bIsRenderTargetCapture)
+			{
+				bResetMediaImage = !CaptureSettings->RenderTargetCaptures.IsValidIndex(RenderTargetOutputInfoIndex) ||
+					CaptureSettings->RenderTargetCaptures[RenderTargetOutputInfoIndex].MediaOutput != GetMediaItem();
+			}
+			else if (bIsActiveEditorViewportCapture)
+			{
+				bResetMediaImage = CaptureSettings->CurrentViewportMediaOutput.MediaOutput != GetMediaItem();
+			}
+
+			if (bResetMediaImage)
+			{
+				ConfigureMediaImage();
+			}
+		}
 	}
 }
 
-bool SMediaProfileMediaOutputDisplay::IsCapturingMediaOutput() const
+TOptional<EMediaCaptureState> SMediaProfileMediaOutputDisplay::GetMediaOutputCaptureState() const
 {
 	TSharedPtr<FMediaProfileEditor> PinnedMediaProfileEditor = MediaProfileEditor.Pin();
 	if (!PinnedMediaProfileEditor.IsValid())
 	{
-		return false;
+		return TOptional<EMediaCaptureState>();
 	}
 
-	return PinnedMediaProfileEditor->IsMediaOutputCapturing(GetMediaItem());
+	return PinnedMediaProfileEditor->GetMediaOutputCaptureState(GetMediaItem());
+}
+
+FOptionalSize SMediaProfileMediaOutputDisplay::GetMediaOutputDesiredAspectRatio() const
+{
+	UMediaOutput* MediaOutput = GetMediaItem();
+	if (!MediaOutput)
+	{
+		return FOptionalSize();
+	}
+	
+	FIntPoint TargetSize = MediaOutput->GetRequestedSize();
+	if (TargetSize == UMediaOutput::RequestCaptureSourceSize)
+	{
+		TargetSize = GetMutableDefault<UMediaFrameworkMediaCaptureSettings>()->DefaultViewportCaptureSize;
+	}
+	
+	const float AspectRatio = (TargetSize.X > 0 && TargetSize.Y > 0) ? (float)TargetSize.X / (float)TargetSize.Y : 1.77777777f;
+	return AspectRatio;
 }
 
 void SMediaProfileMediaOutputDisplay::OnCaptureMethodChanged(UMediaOutput* MediaOutput)
diff --git a/Engine/Plugins/Media/MediaFrameworkUtilities/Source/MediaFrameworkUtilitiesEditor/Private/AssetEditor/SMediaProfileMediaItemDisplay.h b/Engine/Plugins/Media/MediaFrameworkUtilities/Source/MediaFrameworkUtilitiesEditor/Private/AssetEditor/SMediaProfileMediaItemDisplay.h
--- a/Engine/Plugins/Media/MediaFrameworkUtilities/Source/MediaFrameworkUtilitiesEditor/Private/AssetEditor/SMediaProfileMediaItemDisplay.h
+++ b/Engine/Plugins/Media/MediaFrameworkUtilities/Source/MediaFrameworkUtilitiesEditor/Private/AssetEditor/SMediaProfileMediaItemDisplay.h
@@ -759,7 +759,8 @@ protected:
 	virtual void OnObjectPropertyChanged(UObject* InObject, FPropertyChangedEvent& InPropertyChangedEvent) override;
 
 private:
-	bool IsCapturingMediaOutput() const;
+	TOptional<EMediaCaptureState> GetMediaOutputCaptureState() const;
+	FOptionalSize GetMediaOutputDesiredAspectRatio() const;
 	
 	void OnCaptureMethodChanged(UMediaOutput* MediaOutput);
 	
