Media Profile: LOCTEXT collision fixes for Current Viewport

diff --git a/Engine/Plugins/Media/MediaFrameworkUtilities/Source/MediaFrameworkUtilitiesEditor/Private/AssetEditor/MediaProfileMediaItemDetailCustomization.cpp b/Engine/Plugins/Media/MediaFrameworkUtilities/Source/MediaFrameworkUtilitiesEditor/Private/AssetEditor/MediaProfileMediaItemDetailCustomization.cpp
--- a/Engine/Plugins/Media/MediaFrameworkUtilities/Source/MediaFrameworkUtilitiesEditor/Private/AssetEditor/MediaProfileMediaItemDetailCustomization.cpp
+++ b/Engine/Plugins/Media/MediaFrameworkUtilities/Source/MediaFrameworkUtilitiesEditor/Private/AssetEditor/MediaProfileMediaItemDetailCustomization.cpp
@@ -96,8 +96,8 @@ public:
 			{
 				FMenuBuilder MenuBuilder(true, nullptr);
 
-				MenuBuilder.AddMenuEntry(LOCTEXT("MediaCaptureViewportMethod", "Current Viewport"),
-					LOCTEXT("MediaCaptureViewportMethodToolTip", "Capture from the current viewport"),
+				MenuBuilder.AddMenuEntry(LOCTEXT("MediaCaptureCurrentViewportMethod", "Current Viewport"),
+					LOCTEXT("MediaCaptureCurrentViewportMethodToolTip", "Capture from the current viewport"),
 					FSlateIcon(),
 					FUIAction(FExecuteAction::CreateSP(this, &SMediaCaptureMethodComboButton::OnCurrentViewportMethodSelected)));
 				
@@ -144,11 +144,6 @@ private:
 			{
 				continue;
 			}
-		
-			// We expect there to be only one capture configuration per media output, so clear
-			// any existing viewport and render target capture configurations before adding the new one
-			MediaCaptureSettings->ViewportCaptures.RemoveAll(GetFindOutputInfoPredicate<FMediaFrameworkCaptureCameraViewportCameraOutputInfo>(MediaOutput));
-			MediaCaptureSettings->RenderTargetCaptures.RemoveAll(GetFindOutputInfoPredicate<FMediaFrameworkCaptureRenderTargetCameraOutputInfo>(MediaOutput));
 			
 			// Now, if the object has an existing viewport or render target capture config, we want to copy that config over to the current viewport config, so find it
 			if (FMediaFrameworkCaptureCameraViewportCameraOutputInfo* ExistingViewportOutputInfo =
@@ -162,6 +157,11 @@ private:
 				MediaCaptureSettings->CurrentViewportMediaOutput.CaptureOptions = ExistingRenderTargetOutputInfo->CaptureOptions;
 			}
 
+			// We expect there to be only one capture configuration per media output, so clear
+			// any existing viewport and render target capture configurations before adding the new one
+			MediaCaptureSettings->ViewportCaptures.RemoveAll(GetFindOutputInfoPredicate<FMediaFrameworkCaptureCameraViewportCameraOutputInfo>(MediaOutput));
+			MediaCaptureSettings->RenderTargetCaptures.RemoveAll(GetFindOutputInfoPredicate<FMediaFrameworkCaptureRenderTargetCameraOutputInfo>(MediaOutput));
+			
 			MediaCaptureSettings->CurrentViewportMediaOutput.MediaOutput = MediaOutput;
 		}
 
