Media Profile: Moved Start Capture button up to details panel header to avoid the need to scroll it into view

diff --git a/Engine/Plugins/Media/MediaFrameworkUtilities/Source/MediaFrameworkUtilitiesEditor/Private/AssetEditor/MediaProfileMediaItemDetailCustomization.cpp b/Engine/Plugins/Media/MediaFrameworkUtilities/Source/MediaFrameworkUtilitiesEditor/Private/AssetEditor/MediaProfileMediaItemDetailCustomization.cpp
--- a/Engine/Plugins/Media/MediaFrameworkUtilities/Source/MediaFrameworkUtilitiesEditor/Private/AssetEditor/MediaProfileMediaItemDetailCustomization.cpp
+++ b/Engine/Plugins/Media/MediaFrameworkUtilities/Source/MediaFrameworkUtilitiesEditor/Private/AssetEditor/MediaProfileMediaItemDetailCustomization.cpp
@@ -475,40 +475,6 @@ void FMediaProfileMediaOutputDetailCustomization::CustomizeDetails(IDetailLayout
 		ECategoryPriority::Important);
 
 	RegisterCaptureSection(FModuleManager::GetModuleChecked<FPropertyEditorModule>("PropertyEditor"), CaptureCategoryName);
-
-	MediaCaptureCategory
-		.AddCustomRow(LOCTEXT("MediaCaptureButtonFilter", "Start/Stop Media Capture"))
-		.Visibility(TAttribute<EVisibility>::CreateSP(this, &FMediaProfileMediaOutputDetailCustomization::GetValidObjectVisibility))
-		.WholeRowContent()
-		[
-			SNew(SButton)
-			.HAlign(HAlign_Center)
-			.VAlign(VAlign_Center)
-			.OnPressed(this, &FMediaProfileMediaOutputDetailCustomization::OnCaptureButtonPressed)
-			[
-				SNew(SHorizontalBox)
-
-				+SHorizontalBox::Slot()
-				.AutoWidth()
-				.HAlign(HAlign_Center)
-				.VAlign(VAlign_Center)
-				.Padding(0.0f, 0.0f, 4.0f, 0.0f)				
-				[
-					SNew(SImage)
-					.Image(this, &FMediaProfileMediaOutputDetailCustomization::GetCaptureButtonImage)
-				]
-
-				+SHorizontalBox::Slot()
-				.AutoWidth()
-				.HAlign(HAlign_Center)
-				.VAlign(VAlign_Center)
-				[
-					SNew(STextBlock)
-					.Text(this, &FMediaProfileMediaOutputDetailCustomization::GetCaptureButtonText)
-				]
-			]
-			.IsEnabled(this, &FMediaProfileMediaOutputDetailCustomization::IsCaptureButtonEnabled)
-		];
 		
 	MediaCaptureCategory
 		.AddCustomRow(LOCTEXT("MediaCaptureMethodFilter", "Capture Method"))
@@ -735,141 +701,6 @@ void FMediaProfileMediaOutputDetailCustomization::RegisterCaptureSection(FProper
 	}
 }
 
-void FMediaProfileMediaOutputDetailCustomization::OnCaptureButtonPressed()
-{
-	if (MediaProfileEditor.IsValid())
-	{
-		TSharedPtr<FMediaProfileEditor> PinnedEditor = MediaProfileEditor.Pin();
-			
-		TArray<TWeakObjectPtr<UObject>> Objects;
-		CachedDetailBuilder.Pin()->GetObjectsBeingCustomized(Objects);
-
-		bool bIsCapturing = false;
-		for (const TWeakObjectPtr<UObject>& WeakObject : Objects)
-		{
-			TStrongObjectPtr<UObject> Object = WeakObject.Pin();
-			if (Object.IsValid() && Object.Get()->IsA<UMediaOutput>())
-			{
-				if (PinnedEditor->IsMediaOutputCapturing(Cast<UMediaOutput>(Object.Get())))
-				{
-					bIsCapturing = true;
-					break;
-				}
-			}
-		}
-			
-		if (bIsCapturing)
-		{
-			for (const TWeakObjectPtr<UObject>& WeakObject : Objects)
-			{
-				TStrongObjectPtr<UObject> Object = WeakObject.Pin();
-				if (Object.IsValid() && Object.Get()->IsA<UMediaOutput>())
-				{
-					PinnedEditor->StopCapturingMediaOutput(Cast<UMediaOutput>(Object.Get()));
-				}
-			}
-		}
-		else
-		{
-			for (const TWeakObjectPtr<UObject>& WeakObject : Objects)
-			{
-			TStrongObjectPtr<UObject> Object = WeakObject.Pin();
-				if (Object.IsValid() && Object.Get()->IsA<UMediaOutput>())
-				{
-					PinnedEditor->StartCapturingMediaOutput(Cast<UMediaOutput>(Object.Get()));
-				}
-			}
-		}
-	}
-}
-
-const FSlateBrush* FMediaProfileMediaOutputDetailCustomization::GetCaptureButtonImage() const
-{
-	bool bIsCapturing = false;
-	if (MediaProfileEditor.IsValid())
-	{
-		TSharedPtr<FMediaProfileEditor> PinnedEditor = MediaProfileEditor.Pin();
-			
-		TArray<TWeakObjectPtr<UObject>> Objects;
-		CachedDetailBuilder.Pin()->GetObjectsBeingCustomized(Objects);
-
-		for (const TWeakObjectPtr<UObject>& WeakObject : Objects)
-		{
-			TStrongObjectPtr<UObject> Object = WeakObject.Pin();
-			if (Object.IsValid() && Object.Get()->IsA<UMediaOutput>())
-			{
-				if (PinnedEditor->IsMediaOutputCapturing(Cast<UMediaOutput>(Object.Get())))
-				{
-					bIsCapturing = true;
-					break;
-				}
-			}
-		}
-	}
-		
-	return !bIsCapturing ? FMediaFrameworkUtilitiesEditorStyle::Get().GetBrush("MediaCapture.Capture") : FAppStyle::GetBrush("Icons.Toolbar.Stop");
-}
-
-FText FMediaProfileMediaOutputDetailCustomization::GetCaptureButtonText() const
-{
-	bool bIsCapturing = false;
-	if (MediaProfileEditor.IsValid())
-	{
-		TSharedPtr<FMediaProfileEditor> PinnedEditor = MediaProfileEditor.Pin();
-			
-		TArray<TWeakObjectPtr<UObject>> Objects;
-		CachedDetailBuilder.Pin()->GetObjectsBeingCustomized(Objects);
-
-		for (const TWeakObjectPtr<UObject>& WeakObject : Objects)
-		{
-			TStrongObjectPtr<UObject> Object = WeakObject.Pin();
-			if (Object.IsValid() && Object.Get()->IsA<UMediaOutput>())
-			{
-				if (PinnedEditor->IsMediaOutputCapturing(Cast<UMediaOutput>(Object.Get())))
-				{
-					bIsCapturing = true;
-					break;
-				}
-			}
-		}
-	}
-		
-	return !bIsCapturing ? LOCTEXT("MediaOutputStartCapture", "Start Capture") : LOCTEXT("MediaOutputStopCapture", "Stop Capture");
-}
-
-bool FMediaProfileMediaOutputDetailCustomization::IsCaptureButtonEnabled() const
-{
-	bool bIsCapturing = false;
-	bool bIsInvalidCaptureConfig = false;
-	if (MediaProfileEditor.IsValid())
-	{
-		TSharedPtr<FMediaProfileEditor> PinnedEditor = MediaProfileEditor.Pin();
-			
-		TArray<TWeakObjectPtr<UObject>> Objects;
-		CachedDetailBuilder.Pin()->GetObjectsBeingCustomized(Objects);
-
-		for (const TWeakObjectPtr<UObject>& WeakObject : Objects)
-		{
-			TStrongObjectPtr<UObject> Object = WeakObject.Pin();
-			if (Object.IsValid() && Object.Get()->IsA<UMediaOutput>())
-			{
-				UMediaOutput* MediaOutput = Cast<UMediaOutput>(Object.Get());
-				if (PinnedEditor->IsMediaOutputCapturing(MediaOutput))
-				{
-					bIsCapturing = true;
-				}
-					
-				if (!PinnedEditor->CanMediaOutputCapture(MediaOutput))
-				{
-					bIsInvalidCaptureConfig = true;
-				}
-			}
-		}
-	}
-		
-	return bIsCapturing || !bIsInvalidCaptureConfig;
-}
-
 void FMediaProfileMediaOutputDetailCustomization::OnCaptureSettingsObjectChanged()
 {
 	if (CachedDetailBuilder.IsValid())
diff --git a/Engine/Plugins/Media/MediaFrameworkUtilities/Source/MediaFrameworkUtilitiesEditor/Private/AssetEditor/MediaProfileMediaItemDetailCustomization.h b/Engine/Plugins/Media/MediaFrameworkUtilities/Source/MediaFrameworkUtilitiesEditor/Private/AssetEditor/MediaProfileMediaItemDetailCustomization.h
--- a/Engine/Plugins/Media/MediaFrameworkUtilities/Source/MediaFrameworkUtilitiesEditor/Private/AssetEditor/MediaProfileMediaItemDetailCustomization.h
+++ b/Engine/Plugins/Media/MediaFrameworkUtilities/Source/MediaFrameworkUtilitiesEditor/Private/AssetEditor/MediaProfileMediaItemDetailCustomization.h
@@ -368,12 +368,6 @@ private:
 	/** Registers a section for the Capture category */
 	void RegisterCaptureSection(FPropertyEditorModule& PropertyModule, const FName& CaptureCategory);
 	
-	void OnCaptureButtonPressed();
-	
-	const FSlateBrush* GetCaptureButtonImage() const;
-	FText GetCaptureButtonText() const;
-	bool IsCaptureButtonEnabled() const;
-	
 	void OnCaptureSettingsObjectChanged();
 
 	void OnCaptureMethodChanged();
diff --git a/Engine/Plugins/Media/MediaFrameworkUtilities/Source/MediaFrameworkUtilitiesEditor/Private/AssetEditor/SMediaProfileDetailsPanel.cpp b/Engine/Plugins/Media/MediaFrameworkUtilities/Source/MediaFrameworkUtilitiesEditor/Private/AssetEditor/SMediaProfileDetailsPanel.cpp
--- a/Engine/Plugins/Media/MediaFrameworkUtilities/Source/MediaFrameworkUtilitiesEditor/Private/AssetEditor/SMediaProfileDetailsPanel.cpp
+++ b/Engine/Plugins/Media/MediaFrameworkUtilities/Source/MediaFrameworkUtilitiesEditor/Private/AssetEditor/SMediaProfileDetailsPanel.cpp
@@ -64,6 +64,40 @@ void SMediaProfileDetailsPanel::Construct(const FArguments& InArgs, const TShare
 		[
 			SAssignNew(InfoPanel, SMediaObjectInfoPanel, OwningEditor, MediaProfile.Get())
 		]
+
+		+SVerticalBox::Slot()
+		.Padding(6.0f, 4.0f)
+		.AutoHeight()
+		[
+			SNew(SButton)
+			.HAlign(HAlign_Center)
+			.VAlign(VAlign_Center)
+			.IsEnabled(this, &SMediaProfileDetailsPanel::IsCaptureButtonEnabled)
+			.Visibility(this, &SMediaProfileDetailsPanel::GetCaptureButtonVisibility)
+			.OnPressed(this, &SMediaProfileDetailsPanel::OnCaptureButtonPressed)
+			[
+				SNew(SHorizontalBox)
+
+				+SHorizontalBox::Slot()
+				.AutoWidth()
+				.HAlign(HAlign_Center)
+				.VAlign(VAlign_Center)
+				.Padding(0.0f, 0.0f, 4.0f, 0.0f)				
+				[
+					SNew(SImage)
+					.Image(this, &SMediaProfileDetailsPanel::GetCaptureButtonImage)
+				]
+
+				+SHorizontalBox::Slot()
+				.AutoWidth()
+				.HAlign(HAlign_Center)
+				.VAlign(VAlign_Center)
+				[
+					SNew(STextBlock)
+					.Text(this, &SMediaProfileDetailsPanel::GetCaptureButtonText)
+				]
+			]
+		]
 		
 		+SVerticalBox::Slot()
 		[
@@ -197,4 +231,175 @@ void SMediaProfileDetailsPanel::RefreshDetailsView()
 	}
 }
 
+void SMediaProfileDetailsPanel::OnCaptureButtonPressed()
+{
+	TSharedPtr<FMediaProfileEditor> PinnedEditor = OwningEditor.Pin();
+	if (!PinnedEditor.IsValid())
+	{
+		return;
+	}
+	
+	TStrongObjectPtr<UMediaProfile> PinnedMediaProfile = MediaProfile.Pin();
+	if (!PinnedMediaProfile.IsValid())
+	{
+		return;
+	}
+	
+	bool bIsCapturing = false;
+	for (int32 OutputIndex : SelectedMediaOutputs)
+	{
+		if (UMediaOutput* MediaOutput = PinnedMediaProfile->GetMediaOutput(OutputIndex))
+		{
+			if (PinnedEditor->IsMediaOutputCapturing(MediaOutput))
+			{
+				bIsCapturing = true;
+				break;
+			}
+		}
+	}
+		
+	if (bIsCapturing)
+	{
+		for (int32 OutputIndex : SelectedMediaOutputs)
+		{
+			if (UMediaOutput* MediaOutput = PinnedMediaProfile->GetMediaOutput(OutputIndex))
+			{
+				PinnedEditor->StopCapturingMediaOutput(MediaOutput);
+			}
+		}
+	}
+	else
+	{
+		for (int32 OutputIndex : SelectedMediaOutputs)
+		{
+			if (UMediaOutput* MediaOutput = PinnedMediaProfile->GetMediaOutput(OutputIndex))
+			{
+				PinnedEditor->StartCapturingMediaOutput(MediaOutput);
+			}
+		}
+	}
+}
+
+const FSlateBrush* SMediaProfileDetailsPanel::GetCaptureButtonImage() const
+{
+	TSharedPtr<FMediaProfileEditor> PinnedEditor = OwningEditor.Pin();
+	if (!PinnedEditor.IsValid())
+	{
+		return nullptr;
+	}
+
+	TStrongObjectPtr<UMediaProfile> PinnedMediaProfile = MediaProfile.Pin();
+	if (!PinnedMediaProfile.IsValid())
+	{
+		return nullptr;
+	}
+	
+	bool bIsCapturing = false;
+	for (int32 OutputIndex : SelectedMediaOutputs)
+	{
+		if (UMediaOutput* MediaOutput = PinnedMediaProfile->GetMediaOutput(OutputIndex))
+		{
+			if (PinnedEditor->IsMediaOutputCapturing(MediaOutput))
+			{
+				bIsCapturing = true;
+				break;
+			}
+		}
+	}
+		
+	return !bIsCapturing ? FMediaFrameworkUtilitiesEditorStyle::Get().GetBrush("MediaCapture.Capture") : FAppStyle::GetBrush("Icons.Toolbar.Stop");
+}
+
+FText SMediaProfileDetailsPanel::GetCaptureButtonText() const
+{
+	TSharedPtr<FMediaProfileEditor> PinnedEditor = OwningEditor.Pin();
+	if (!PinnedEditor.IsValid())
+	{
+		return FText::GetEmpty();
+	}
+
+	TStrongObjectPtr<UMediaProfile> PinnedMediaProfile = MediaProfile.Pin();
+	if (!PinnedMediaProfile.IsValid())
+	{
+		return FText::GetEmpty();
+	}
+	
+	bool bIsCapturing = false;
+	for (int32 OutputIndex : SelectedMediaOutputs)
+	{
+		if (UMediaOutput* MediaOutput = PinnedMediaProfile->GetMediaOutput(OutputIndex))
+		{
+			if (PinnedEditor->IsMediaOutputCapturing(MediaOutput))
+			{
+				bIsCapturing = true;
+				break;
+			}
+		}
+	}
+	
+	return !bIsCapturing ? LOCTEXT("MediaOutputStartCapture", "Start Capture") : LOCTEXT("MediaOutputStopCapture", "Stop Capture");
+}
+
+EVisibility SMediaProfileDetailsPanel::GetCaptureButtonVisibility() const
+{
+	TSharedPtr<FMediaProfileEditor> PinnedEditor = OwningEditor.Pin();
+	if (!PinnedEditor.IsValid())
+	{
+		return EVisibility::Collapsed;
+	}
+
+	TStrongObjectPtr<UMediaProfile> PinnedMediaProfile = MediaProfile.Pin();
+	if (!PinnedMediaProfile.IsValid())
+	{
+		return EVisibility::Collapsed;
+	}
+
+	bool bHasMediaOutputs = false;
+	for (int32 OutputIndex : SelectedMediaOutputs)
+	{
+		if (OutputIndex > INDEX_NONE && OutputIndex < PinnedMediaProfile->NumMediaOutputs())
+		{
+			bHasMediaOutputs = true;
+			break;
+		}
+	}
+
+	return bHasMediaOutputs ? EVisibility::Visible : EVisibility::Collapsed;
+}
+
+bool SMediaProfileDetailsPanel::IsCaptureButtonEnabled() const
+{
+	TSharedPtr<FMediaProfileEditor> PinnedEditor = OwningEditor.Pin();
+	if (!PinnedEditor.IsValid())
+	{
+		return false;
+	}
+
+	TStrongObjectPtr<UMediaProfile> PinnedMediaProfile = MediaProfile.Pin();
+	if (!PinnedMediaProfile.IsValid())
+	{
+		return false;
+	}
+	
+	bool bIsCapturing = false;
+	bool bIsInvalidCaptureConfig = false;
+	for (int32 OutputIndex : SelectedMediaOutputs)
+	{
+		if (UMediaOutput* MediaOutput = PinnedMediaProfile->GetMediaOutput(OutputIndex))
+		{
+			if (PinnedEditor->IsMediaOutputCapturing(MediaOutput))
+			{
+				bIsCapturing = true;
+			}
+					
+			if (!PinnedEditor->CanMediaOutputCapture(MediaOutput))
+			{
+				bIsInvalidCaptureConfig = true;
+			}
+		}
+	}
+	
+	return bIsCapturing || !bIsInvalidCaptureConfig;
+}
+
 #undef LOCTEXT_NAMESPACE
\ No newline at end of file
diff --git a/Engine/Plugins/Media/MediaFrameworkUtilities/Source/MediaFrameworkUtilitiesEditor/Private/AssetEditor/SMediaProfileDetailsPanel.h b/Engine/Plugins/Media/MediaFrameworkUtilities/Source/MediaFrameworkUtilitiesEditor/Private/AssetEditor/SMediaProfileDetailsPanel.h
--- a/Engine/Plugins/Media/MediaFrameworkUtilities/Source/MediaFrameworkUtilitiesEditor/Private/AssetEditor/SMediaProfileDetailsPanel.h
+++ b/Engine/Plugins/Media/MediaFrameworkUtilities/Source/MediaFrameworkUtilitiesEditor/Private/AssetEditor/SMediaProfileDetailsPanel.h
@@ -36,6 +36,21 @@ public:
 
 private:
 	void RefreshDetailsView();
+
+	/** Raised when the Start/Stop Capture button is pressed */
+	void OnCaptureButtonPressed();
+
+	/** Gets the icon to display in the capture button, depending on if there is a live capture or not */
+	const FSlateBrush* GetCaptureButtonImage() const;
+
+	/** Gets the text to display in the capture button, depending on if there is a live capture or not */
+	FText GetCaptureButtonText() const;
+
+	/** Gets the visibility of the capture button, depending on the types of the objects being displayed in the details panel */
+	EVisibility GetCaptureButtonVisibility() const;
+
+	/** Gets whether the capture button is enabled, depending on whether the selected outputs are properly configured to capture */
+	bool IsCaptureButtonEnabled() const;
 	
 private:
 	/** The editor that owns this details panel */
