Media Profile: Changed behavior of media profile toolbar entry to avoid users having to select a subclass when creating a new media profile asset, as well as allowing creation when the default button is clicked if no media profiles are configured

diff --git a/Engine/Plugins/Media/MediaFrameworkUtilities/Source/MediaFrameworkUtilitiesEditor/Private/UI/MediaProfileMenuEntry.cpp b/Engine/Plugins/Media/MediaFrameworkUtilities/Source/MediaFrameworkUtilitiesEditor/Private/UI/MediaProfileMenuEntry.cpp
--- a/Engine/Plugins/Media/MediaFrameworkUtilities/Source/MediaFrameworkUtilitiesEditor/Private/UI/MediaProfileMenuEntry.cpp
+++ b/Engine/Plugins/Media/MediaFrameworkUtilities/Source/MediaFrameworkUtilitiesEditor/Private/UI/MediaProfileMenuEntry.cpp
@@ -39,8 +39,11 @@ struct FMediaProfileMenuEntryImpl
 				{
 					GEditor->GetEditorSubsystem<UAssetEditorSubsystem>()->OpenEditorForAsset(MediaProfile);
 				}
-			}),
-			FCanExecuteAction::CreateLambda([] { return IMediaProfileManager::Get().GetCurrentMediaProfile() != nullptr; })
+				else
+				{
+					CreateNewProfile();
+				}
+			})
 		);
 
 		UToolMenu* Menu = UToolMenus::Get()->ExtendMenu("LevelEditor.LevelEditorToolBar.User");
@@ -51,7 +54,7 @@ struct FMediaProfileMenuEntryImpl
 			UMediaProfile* MediaProfile = IMediaProfileManager::Get().GetCurrentMediaProfile();
 			if (MediaProfile == nullptr)
 			{
-				return LOCTEXT("EmptyMediaProfile_ToolTip", "Select a Media Profile to edit it.");
+				return LOCTEXT("EmptyMediaProfile_ToolTip", "Create a new media profile");
 			}
 			return FText::Format(LOCTEXT("MediaProfile_ToolTip", "Edit '{0}'")
 				, FText::FromName(MediaProfile->GetFName()));
@@ -96,37 +99,14 @@ struct FMediaProfileMenuEntryImpl
 
 	void CreateNewProfile()
 	{
-		class FModifierClassFilter : public IClassViewerFilter
+		UMediaProfileFactoryNew* FactoryInstance = DuplicateObject<UMediaProfileFactoryNew>(GetDefault<UMediaProfileFactoryNew>(), GetTransientPackage());
+		FAssetToolsModule& AssetToolsModule = FAssetToolsModule::GetModule();
+		UMediaProfile* NewAsset = Cast<UMediaProfile>(FAssetToolsModule::GetModule().Get().CreateAssetWithDialog(UMediaProfile::StaticClass(), FactoryInstance));
+		if (NewAsset != nullptr)
 		{
-		public:
-			bool IsClassAllowed(const FClassViewerInitializationOptions& InInitOptions, const UClass* InClass, TSharedRef< FClassViewerFilterFuncs > InFilterFuncs) override
-			{
-				return InClass->IsChildOf(UMediaProfile::StaticClass()) && !InClass->HasAnyClassFlags(CLASS_Abstract | CLASS_Deprecated | CLASS_NewerVersionExists | CLASS_HideDropDown);
-			}
-
-			virtual bool IsUnloadedClassAllowed(const FClassViewerInitializationOptions& InInitOptions, const TSharedRef< const IUnloadedBlueprintData > InClass, TSharedRef< FClassViewerFilterFuncs > InFilterFuncs) override
-			{
-				return InClass->IsChildOf(UMediaProfile::StaticClass()) && !InClass->HasAnyClassFlags(CLASS_Abstract | CLASS_Deprecated | CLASS_NewerVersionExists | CLASS_HideDropDown);
-			}
-		};
-
-		const FText TitleText = LOCTEXT("CreateMediaProfileOptions", "Pick Media Profile Class");
-		FClassViewerInitializationOptions Options;
-		Options.ClassFilters.Add(MakeShared<FModifierClassFilter>());
-		UClass* ChosenClass = nullptr;
-		const bool bPressedOk = SClassPickerDialog::PickClass(TitleText, Options, ChosenClass, UMediaProfile::StaticClass());
-
-		if (bPressedOk && ChosenClass != nullptr)
-		{
-			UMediaProfileFactoryNew* FactoryInstance = DuplicateObject<UMediaProfileFactoryNew>(GetDefault<UMediaProfileFactoryNew>(), GetTransientPackage());
-			FAssetToolsModule& AssetToolsModule = FAssetToolsModule::GetModule();
-			UMediaProfile* NewAsset = Cast<UMediaProfile>(FAssetToolsModule::GetModule().Get().CreateAssetWithDialog(ChosenClass, FactoryInstance));
-			if (NewAsset != nullptr)
-			{
-				GetMutableDefault<UMediaProfileEditorSettings>()->SetUserMediaProfile(NewAsset);
-				IMediaProfileManager::Get().SetCurrentMediaProfile(NewAsset);
-				GEditor->GetEditorSubsystem<UAssetEditorSubsystem>()->OpenEditorForAsset(NewAsset);
-			}
+			GetMutableDefault<UMediaProfileEditorSettings>()->SetUserMediaProfile(NewAsset);
+			IMediaProfileManager::Get().SetCurrentMediaProfile(NewAsset);
+			GEditor->GetEditorSubsystem<UAssetEditorSubsystem>()->OpenEditorForAsset(NewAsset);
 		}
 	}
 
