Media Profile: Changed default media capture to Current Viewport when a new media output is set

diff --git a/Engine/Plugins/Media/MediaFrameworkUtilities/Source/MediaFrameworkUtilitiesEditor/Private/AssetEditor/MediaProfileMediaItemDetailCustomization.cpp b/Engine/Plugins/Media/MediaFrameworkUtilities/Source/MediaFrameworkUtilitiesEditor/Private/AssetEditor/MediaProfileMediaItemDetailCustomization.cpp
--- a/Engine/Plugins/Media/MediaFrameworkUtilities/Source/MediaFrameworkUtilitiesEditor/Private/AssetEditor/MediaProfileMediaItemDetailCustomization.cpp
+++ b/Engine/Plugins/Media/MediaFrameworkUtilities/Source/MediaFrameworkUtilitiesEditor/Private/AssetEditor/MediaProfileMediaItemDetailCustomization.cpp
@@ -678,9 +678,7 @@ void FMediaProfileMediaOutputDetailCustomization::SetMediaObject(UMediaOutput* I
 
 	if (!bHasExistingCaptureInfo)
 	{
-		FMediaFrameworkCaptureCameraViewportCameraOutputInfo NewCameraOutputInfo;
-		NewCameraOutputInfo.MediaOutput = InNewMediaObject;
-		CaptureSettings->ViewportCaptures.Add(NewCameraOutputInfo);
+		CaptureSettings->CurrentViewportMediaOutput.MediaOutput = InNewMediaObject;
 	}
 	
 	PinnedMediaProfile->SetMediaOutput(Index, InNewMediaObject);
