MediaProfile: Avoid crash on null media source.

diff --git a/Engine/Plugins/Media/MediaFrameworkUtilities/Source/MediaFrameworkUtilities/Private/Profile/MediaProfilePlaybackManager.cpp b/Engine/Plugins/Media/MediaFrameworkUtilities/Source/MediaFrameworkUtilities/Private/Profile/MediaProfilePlaybackManager.cpp
--- a/Engine/Plugins/Media/MediaFrameworkUtilities/Source/MediaFrameworkUtilities/Private/Profile/MediaProfilePlaybackManager.cpp
+++ b/Engine/Plugins/Media/MediaFrameworkUtilities/Source/MediaFrameworkUtilities/Private/Profile/MediaProfilePlaybackManager.cpp
@@ -114,6 +114,11 @@ UMediaTexture* UMediaProfilePlaybackManager::OpenSourceFromIndex(int32 InMediaSo
 	}
 	
 	UMediaSource* MediaSource = MediaProfile->GetMediaSource(InMediaSourceIndex);
+	if (!MediaSource)
+	{
+		return nullptr;
+	}
+
 	UMediaPlayer* MediaSourcePlayer = FindOrCreateMediaPlayerForIndex(InMediaSourceIndex);
 
 	MediaTexture->SetMediaPlayer(MediaSourcePlayer);
