Media Profile: Moved media capture management from media profile editor into media profile playback manager so that media captures could continue even when the media profile editor window is closed. Also fixed issue where Current Viewport was not correctly resizing the viewport when capturing

diff --git a/Engine/Plugins/Media/MediaFrameworkUtilities/Source/MediaFrameworkUtilities/MediaFrameworkUtilities.Build.cs b/Engine/Plugins/Media/MediaFrameworkUtilities/Source/MediaFrameworkUtilities/MediaFrameworkUtilities.Build.cs
--- a/Engine/Plugins/Media/MediaFrameworkUtilities/Source/MediaFrameworkUtilities/MediaFrameworkUtilities.Build.cs
+++ b/Engine/Plugins/Media/MediaFrameworkUtilities/Source/MediaFrameworkUtilities/MediaFrameworkUtilities.Build.cs
@@ -33,6 +33,7 @@ namespace UnrealBuildTool.Rules
 					new string[]
 					{
 						"EditorFramework",
+						"LevelEditor",
                         "MediaPlayerEditor",
                         "Slate",
 						"SlateCore",
diff --git a/Engine/Plugins/Media/MediaFrameworkUtilities/Source/MediaFrameworkUtilities/Private/Profile/MediaProfile.cpp b/Engine/Plugins/Media/MediaFrameworkUtilities/Source/MediaFrameworkUtilities/Private/Profile/MediaProfile.cpp
--- a/Engine/Plugins/Media/MediaFrameworkUtilities/Source/MediaFrameworkUtilities/Private/Profile/MediaProfile.cpp
+++ b/Engine/Plugins/Media/MediaFrameworkUtilities/Source/MediaFrameworkUtilities/Private/Profile/MediaProfile.cpp
@@ -292,6 +292,7 @@ bool UMediaProfile::RemoveMediaOutput(int32 Index)
 		return false;
 	}
 
+	PlaybackManager->ChangeMediaOutputIndex(Index, INDEX_NONE);
 	MediaOutputs.RemoveAt(Index, EAllowShrinking::No);
 
 #if WITH_EDITOR
@@ -311,6 +312,7 @@ bool UMediaProfile::MoveMediaOutput(int32 CurrentIndex, int32 DestIndex)
 		return false;
 	}
 
+	PlaybackManager->ChangeMediaOutputIndex(CurrentIndex, DestIndex);
 	TObjectPtr<UMediaOutput> MediaOutput = MediaOutputs[CurrentIndex];
 	MediaOutputs.RemoveAt(CurrentIndex, EAllowShrinking::No);
 	MediaOutputs.Insert(MediaOutput, DestIndex);
diff --git a/Engine/Plugins/Media/MediaFrameworkUtilities/Source/MediaFrameworkUtilities/Private/Profile/MediaProfilePlaybackManager.cpp b/Engine/Plugins/Media/MediaFrameworkUtilities/Source/MediaFrameworkUtilities/Private/Profile/MediaProfilePlaybackManager.cpp
--- a/Engine/Plugins/Media/MediaFrameworkUtilities/Source/MediaFrameworkUtilities/Private/Profile/MediaProfilePlaybackManager.cpp
+++ b/Engine/Plugins/Media/MediaFrameworkUtilities/Source/MediaFrameworkUtilities/Private/Profile/MediaProfilePlaybackManager.cpp
@@ -3,15 +3,336 @@
 
 #include "Profile/MediaProfilePlaybackManager.h"
 
+#include "Engine/TextureRenderTarget2D.h"
+#include "MediaCapture.h"
+#include "MediaOutput.h"
 #include "MediaPlayer.h"
 #include "MediaSource.h"
 #include "MediaTexture.h"
+#include "Modules/ModuleManager.h"
 #include "Profile/MediaProfile.h"
+#include "Slate/SceneViewport.h"
 
+#if WITH_EDITOR
+#include "EditorViewportClient.h"
+#include "Framework/Application/SlateApplication.h"
+#include "IAssetViewport.h"
+#include "LevelEditor.h"
+#include "LevelEditorViewport.h"
+#include "Widgets/SViewport.h"
+#endif
+
 namespace MediaProfilePlaybackManager
 {
 	const FString DefaultMediaTexturesPath = TEXT("/MediaFrameworkUtilities/DefaultMediaTextures");
 	const FString DefaultMediaTextureName = TEXT("MediaProfileTexture");
+
+#if WITH_EDITOR
+	class FMediaCaptureEditorViewportClient : public FLevelEditorViewportClient
+	{
+	private:
+		using Super = FLevelEditorViewportClient;
+
+	public:
+		FMediaCaptureEditorViewportClient(EViewModeIndex InViewModeIndex)
+			: FLevelEditorViewportClient(nullptr)
+			, ViewModeIndex(InViewModeIndex)
+		{
+			bSetListenerPosition = false;
+
+			// Default to "game" show flags for camera previews
+			EngineShowFlags = FEngineShowFlags(ESFIM_Game);
+			LastEngineShowFlags = FEngineShowFlags(ESFIM_Editor);
+
+			ViewportType = LVT_Perspective;
+			bDrawAxes = false;
+			bDisableInput = true;
+			SetAllowCinematicControl(false);
+			VisibilityDelegate.BindLambda([] { return true; });
+			AutoSetPIE();
+		}
+
+		virtual ~FMediaCaptureEditorViewportClient() override
+		{
+			Viewport = nullptr;
+			
+			if (GEditor)
+			{
+				FWorldContext* PIEWorldContext = GEditor->GetPIEWorldContext();
+				if (PIEWorldContext)
+				{
+					RemoveReferenceToWorldContext(*PIEWorldContext);
+				}
+			}
+		}
+		
+		virtual UWorld* GetWorld() const override
+		{
+			if (GEditor->PlayWorld)
+			{
+				return GEditor->PlayWorld;
+			}
+			return Super::GetWorld();
+		}
+
+		void SetPIE(bool bInIsPIE)
+		{
+			FWorldContext* PIEWorldContext = GEditor->GetPIEWorldContext();
+			if (bInIsPIE)
+			{
+				if (PIEWorldContext)
+				{
+					RemoveReferenceToWorldContext(GEditor->GetEditorWorldContext());
+					SetReferenceToWorldContext(*PIEWorldContext);
+				}
+			}
+			else
+			{
+				if (PIEWorldContext)
+				{
+					RemoveReferenceToWorldContext(*PIEWorldContext);
+				}
+				SetReferenceToWorldContext(GEditor->GetEditorWorldContext());
+			}
+			SetViewMode(ViewModeIndex);
+			SetRealtime(true);
+		}
+
+		void AutoSetPIE()
+		{
+			if (GEditor->PlayWorld)
+			{
+				SetPIE(true);
+			}
+			else
+			{
+				SetViewMode(ViewModeIndex);
+				SetRealtime(true);
+			}
+		}
+
+	private:
+		EViewModeIndex ViewModeIndex;
+	};
+#else
+	// TODO: If media capture in a Game build ever becomes supported, this should be filled out as appropriate
+	class FMediaCaptureGameViewportClient : public FCommonViewportClient
+	{
+		
+	};
+#endif
+
+	/** Stores and manages all the necessary objects for a managed capture viewport */
+	class FManagedViewport
+	{
+	public:
+		/** The media output the managed viewport is for */
+		TWeakObjectPtr<UMediaOutput> MediaOutput;
+
+		/** The viewport client of the viewport */
+		TSharedPtr<FViewportClient> ViewportClient;
+
+		/** The viewport being managed */
+		TSharedPtr<FSceneViewport> Viewport;
+
+#if WITH_EDITOR
+		/** Scene viewports needs a viewport widget to properly function when in editor */
+		TSharedPtr<SViewport> ViewportWidget;
+#endif
+
+		/** A list of consumers that are actively using the managed viewport */
+		TSet<void*> Consumers;
+		
+		FManagedViewport(UMediaOutput* InMediaOutput)
+			: MediaOutput(InMediaOutput)
+		{ }
+	};
+
+	/** Creates a new managed viewport for capturing */
+	TSharedPtr<FManagedViewport> CreateManagedViewport(UMediaOutput* InMediaOutput, EViewModeIndex InViewMode)
+	{
+		TSharedPtr<FManagedViewport> ManagedViewport = MakeShared<FManagedViewport>(InMediaOutput);
+		
+#if WITH_EDITOR
+		TSharedPtr<FMediaCaptureEditorViewportClient> ViewportClient = MakeShared<FMediaCaptureEditorViewportClient>(InViewMode);
+		
+		ManagedViewport->ViewportClient = ViewportClient;
+		ManagedViewport->ViewportWidget = SNew(SViewport)
+			.RenderDirectlyToWindow(false)
+			.IsEnabled(FSlateApplication::Get().GetNormalExecutionAttribute())
+			.EnableGammaCorrection(false)
+			.EnableBlending(false);
+		
+		ManagedViewport->Viewport = MakeShared<FSceneViewport>(ViewportClient.Get(), ManagedViewport->ViewportWidget);
+		ManagedViewport->ViewportWidget->SetViewportInterface(ManagedViewport->Viewport.ToSharedRef());
+
+		ViewportClient->Viewport = ManagedViewport->Viewport.Get();
+#else
+		TSharedPtr<FMediaCaptureGameViewportClient> ViewportClient = MakeShared<FMediaCaptureGameViewportClient>();
+
+		ManagedViewport->ViewportClient = ViewportClient;
+		ManagedViewport->Viewport = MakeShared<FSceneViewport>(ViewportClient.Get(), nullptr);
+#endif
+
+		return ManagedViewport;
+	}
+
+	/** Searches for an active scene viewport in the engine */
+	TSharedPtr<FSceneViewport> FindActiveViewport()
+	{
+		for (const FWorldContext& Context : GEngine->GetWorldContexts())
+		{
+#if WITH_EDITOR
+			// If there is an active PIE viewport, return it
+			if (Context.WorldType == EWorldType::PIE)
+			{
+				UEditorEngine* EditorEngine = CastChecked<UEditorEngine>(GEngine);
+				FSlatePlayInEditorInfo& Info = EditorEngine->SlatePlayInEditorMap.FindChecked(Context.ContextHandle);
+				if (Info.SlatePlayInEditorWindowViewport.IsValid())
+				{
+					return Info.SlatePlayInEditorWindowViewport;
+				}
+			}
+#endif
+
+			// TODO: If media capture in a Game build ever becomes supported, an appropriate game viewport should be found and returned
+		}
+
+		return nullptr;
+	}
+
+#if WITH_EDITOR
+	/** Searches for a level editor viewport */
+	TSharedPtr<IAssetViewport> FindActiveViewportInterface()
+	{
+		FLevelEditorModule& LevelEditorModule = FModuleManager::Get().GetModuleChecked<FLevelEditorModule>(TEXT("LevelEditor"));
+		return LevelEditorModule.GetFirstActiveViewport();
+	}
+#endif
+
+	/** Stores references to an existing scene viewport managed by the engine or editor, and stores cached view state for that viewport */
+	class FActiveViewport
+	{
+	public:
+		FActiveViewport(UMediaOutput* InMediaOutput, const FMediaCaptureOptions& InCaptureOptions)
+			: MediaOutput(InMediaOutput)
+		{
+			Viewport = FindActiveViewport();
+			
+#if WITH_EDITOR
+			if (!Viewport.IsValid())
+			{
+				if (TSharedPtr<IAssetViewport> ActiveViewportInterface = FindActiveViewportInterface())
+				{
+					ViewportInterface = ActiveViewportInterface;
+				
+					FEditorViewportClient& ViewportClient = ActiveViewportInterface->GetAssetViewportClient();
+
+					// Cache the viewport's current settings, then set them to more appropriate settings for recording
+					CachedViewportFlags.bRealTime = ViewportClient.IsRealtime();
+					CachedViewportFlags.bSetListenerPosition = ViewportClient.bSetListenerPosition;
+					CachedViewportFlags.bDrawAxes = ViewportClient.bDrawAxes;
+					CachedViewportFlags.bDisableInput = ViewportClient.bDisableInput;
+					CachedViewportFlags.bAllowCinematicControl = ViewportClient.AllowsCinematicControl();
+					CachedViewportFlags.VisibilityDelegate = ViewportClient.VisibilityDelegate;
+
+					ViewportClient.SetRealtime(true);
+					ViewportClient.bSetListenerPosition = false;
+					ViewportClient.bDrawAxes = false;
+					ViewportClient.SetAllowCinematicControl(false);
+					ViewportClient.VisibilityDelegate.BindLambda([] { return true; });
+
+					const bool bDisableInput = InCaptureOptions.ResizeMethod != EMediaCaptureResizeMethod::ResizeInRenderPass;
+					ViewportClient.bDisableInput = bDisableInput ? true : ViewportClient.bDisableInput;
+
+					if (InCaptureOptions.ResizeMethod == EMediaCaptureResizeMethod::ResizeSource)
+					{
+						if (TSharedPtr<FSceneViewport> LevelViewport = ActiveViewportInterface->GetSharedActiveViewport())
+						{
+							if (LevelViewport->HasFixedSize())
+							{
+								CachedViewportFlags.ViewportSize = LevelViewport->GetSize();
+							}
+						
+							FIntPoint TargetSize = MediaOutput->GetRequestedSize();
+							if (TargetSize == UMediaOutput::RequestCaptureSourceSize)
+							{
+								TargetSize = FIntPoint(1920, 1080);
+							}
+						
+							LevelViewport->SetFixedViewportSize(TargetSize.X, TargetSize.Y);
+						}
+					}
+				}
+			}
+#endif
+		}
+
+		~FActiveViewport()
+		{
+#if WITH_EDITOR
+			if (TSharedPtr<IAssetViewport> PinnedViewportInterface = ViewportInterface.Pin())
+			{
+				FEditorViewportClient& ViewportClient = PinnedViewportInterface->GetAssetViewportClient();
+
+				// Reset settings
+				ViewportClient.SetRealtime(CachedViewportFlags.bRealTime);
+				ViewportClient.bSetListenerPosition = CachedViewportFlags.bSetListenerPosition;
+				ViewportClient.bDrawAxes = CachedViewportFlags.bDrawAxes;
+				ViewportClient.bDisableInput = CachedViewportFlags.bDisableInput;
+				ViewportClient.SetAllowCinematicControl(CachedViewportFlags.bAllowCinematicControl);
+				ViewportClient.VisibilityDelegate = CachedViewportFlags.VisibilityDelegate;
+
+				if (TSharedPtr<FSceneViewport> LevelViewport = PinnedViewportInterface->GetSharedActiveViewport())
+				{
+					LevelViewport->SetFixedViewportSize(CachedViewportFlags.ViewportSize.X, CachedViewportFlags.ViewportSize.Y);
+				}
+			}
+#endif
+		}
+
+		TSharedPtr<FSceneViewport> GetSceneViewport() const
+		{
+			if (TSharedPtr<FSceneViewport> PinnedActiveViewport = Viewport.Pin())
+			{
+				return PinnedActiveViewport;
+			}
+
+#if WITH_EDITOR
+			if (TSharedPtr<IAssetViewport> PinnedViewportInterface = ViewportInterface.Pin())
+			{
+				return PinnedViewportInterface->GetSharedActiveViewport();
+			}
+#endif
+
+			return nullptr;
+		}
+		
+		/** The media output that is using this active viewport */
+		TWeakObjectPtr<UMediaOutput> MediaOutput;
+
+		/** Weak reference to an existing scene viewport, if one was found */
+		TWeakPtr<FSceneViewport> Viewport;
+
+#if WITH_EDITOR
+		/** Weak reference to a level editor viewport, if one was found */
+		TWeakPtr<IAssetViewport> ViewportInterface;
+
+		/** Cached view state of the active viewport client, to be re-instated once the viewport is no longer being captured */
+		struct
+		{
+			FIntPoint ViewportSize = FIntPoint::ZeroValue;
+			bool bRealTime = false;
+			bool bSetListenerPosition = false;
+			bool bDrawAxes = false;
+			bool bDisableInput = false;
+			bool bAllowCinematicControl = false;
+			FEngineShowFlags EngineShowFlags = ESFIM_Editor;
+			FEngineShowFlags LastEngineShowFlags = ESFIM_Game;
+			FViewportStateGetter VisibilityDelegate;
+		} CachedViewportFlags;
+#endif
+	};
 }
 
 UMediaProfilePlaybackManager::UMediaProfilePlaybackManager(const FObjectInitializer& ObjectInitializer)
@@ -31,6 +352,29 @@
 
 		MediaSourceTextures.Add(TSoftObjectPtr<UMediaTexture>(FSoftObjectPath(Path)));
 	}
+
+	// Close all open sources and outputs when the engine shuts down. If we are in an editor, perform the cleanup a bit earlier
+	// than the Exit process to avoid issues with live viewport clients during the editor cleanup process
+#if WITH_EDITOR
+	if (GEditor && !IsTemplate())
+	{
+		GEditor->OnEditorClose().AddUObject(this, &UMediaProfilePlaybackManager::Cleanup);
+	}
+#else
+	FCoreDelegates::OnEnginePreExit.AddUObject(this, &UMediaProfilePlaybackManager::Cleanup);
+#endif
 }
 
+UMediaProfilePlaybackManager::~UMediaProfilePlaybackManager()
+{
+#if WITH_EDITOR
+	if (GEditor && !IsTemplate())
+	{
+		GEditor->OnEditorClose().RemoveAll(this);
+	}
+#else
+	FCoreDelegates::OnEnginePreExit.RemoveAll(this);
+#endif
+}
+
 UMediaTexture* UMediaProfilePlaybackManager::GetSourceMediaTextureFromIndex(int32 InMediaSourceIndex) const
@@ -265,6 +609,534 @@
 	}
 }
 
+TSharedPtr<FViewportClient> UMediaProfilePlaybackManager::GetOrCreateManagedViewportFromIndex(int32 InMediaOutputIndex, EViewModeIndex InViewMode, void* Consumer)
+{
+	using namespace MediaProfilePlaybackManager;
+
+	UMediaProfile* MediaProfile = GetOwningMediaProfile();
+	if (!MediaProfile)
+	{
+		return nullptr;
+	}
+
+	UMediaOutput* MediaOutput = MediaProfile->GetMediaOutput(InMediaOutputIndex);
+	if (!IsValid(MediaOutput))
+	{
+		return nullptr;
+	}
+	
+	int32 ExistingViewportIndex = ManagedOutputViewports.IndexOfByPredicate([MediaOutput](const TSharedPtr<FManagedViewport>& InManagedViewport)
+	{
+		return InManagedViewport->MediaOutput == MediaOutput;
+	});
+
+	if (ExistingViewportIndex != INDEX_NONE)
+	{
+		if (Consumer != nullptr)
+		{
+			ManagedOutputViewports[ExistingViewportIndex]->Consumers.Add(Consumer);
+		}
+		
+		return ManagedOutputViewports[ExistingViewportIndex]->ViewportClient;
+	}
+
+#if WITH_EDITOR
+	if (ManagedOutputViewports.IsEmpty())
+	{
+		FEditorDelegates::PostPIEStarted.AddUObject(this, &UMediaProfilePlaybackManager::OnPostPIEStarted);
+		FEditorDelegates::PrePIEEnded.AddUObject(this, &UMediaProfilePlaybackManager::OnPrePIEEnded);
+	}
+#endif
+	
+	TSharedPtr<FManagedViewport> NewManagedViewport = CreateManagedViewport(MediaOutput, InViewMode);
+	
+	if (Consumer != nullptr)
+	{
+		NewManagedViewport->Consumers.Add(Consumer);
+	}
+
+	ManagedOutputViewports.Add(NewManagedViewport);
+
+	return NewManagedViewport->ViewportClient;
+}
+
+TSharedPtr<FViewportClient> UMediaProfilePlaybackManager::GetOrCreateManagedViewport(UMediaOutput* InMediaOutput, EViewModeIndex InViewMode, void* Consumer)
+{
+	UMediaProfile* MediaProfile = GetOwningMediaProfile();
+	if (!MediaProfile)
+	{
+		return nullptr;
+	}
+	
+	int32 MediaOutputIndex = MediaProfile->FindMediaOutputIndex(InMediaOutput);
+	return GetOrCreateManagedViewportFromIndex(MediaOutputIndex, InViewMode, Consumer);
+}
+
+void UMediaProfilePlaybackManager::ReleaseManagedViewportFromIndex(int32 InMediaOutputIndex, void* Consumer)
+{
+	using namespace MediaProfilePlaybackManager;
+
+	UMediaProfile* MediaProfile = GetOwningMediaProfile();
+	if (!MediaProfile)
+	{
+		return;
+	}
+
+	UMediaOutput* MediaOutput = MediaProfile->GetMediaOutput(InMediaOutputIndex);
+	if (!IsValid(MediaOutput))
+	{
+		return;
+	}
+	
+	int32 ExistingViewportIndex = ManagedOutputViewports.IndexOfByPredicate([MediaOutput](const TSharedPtr<FManagedViewport>& InManagedViewport)
+	{
+		return InManagedViewport->MediaOutput == MediaOutput;
+	});
+
+	if (ExistingViewportIndex != INDEX_NONE)
+	{
+		ManagedOutputViewports[ExistingViewportIndex]->Consumers.Remove(Consumer);
+		if (ManagedOutputViewports[ExistingViewportIndex]->Consumers.IsEmpty() && !IsOutputCapturingFromIndex(InMediaOutputIndex))
+		{
+			ManagedOutputViewports.RemoveAt(ExistingViewportIndex);
+
+#if WITH_EDITOR
+			if (ManagedOutputViewports.IsEmpty())
+			{
+				FEditorDelegates::PostPIEStarted.RemoveAll(this);
+				FEditorDelegates::PrePIEEnded.RemoveAll(this);
+			}
+#endif
+		}
+	}
+}
+
+void UMediaProfilePlaybackManager::ReleaseManagedViewport(UMediaOutput* InMediaOutput, void* Consumer)
+{
+	UMediaProfile* MediaProfile = GetOwningMediaProfile();
+	if (!MediaProfile)
+	{
+		return;
+	}
+	
+	int32 MediaOutputIndex = MediaProfile->FindMediaOutputIndex(InMediaOutput);
+	return ReleaseManagedViewportFromIndex(MediaOutputIndex, Consumer);
+}
+
+UMediaCapture* UMediaProfilePlaybackManager::OpenManagedViewportOutputFromIndex(int32 InMediaOutputIndex, const FMediaCaptureOptions& InCaptureOptions, bool bAutoRestartCapture)
+{
+	using namespace MediaProfilePlaybackManager;
+
+	UMediaProfile* MediaProfile = GetOwningMediaProfile();
+	if (!MediaProfile)
+	{
+		return nullptr;
+	}
+
+	UMediaOutput* MediaOutput = MediaProfile->GetMediaOutput(InMediaOutputIndex);
+	if (!IsValid(MediaOutput))
+	{
+		return nullptr;
+	}
+	
+	int32 ManagedViewportIndex = ManagedOutputViewports.IndexOfByPredicate([MediaOutput](const TSharedPtr<FManagedViewport>& InManagedViewport)
+	{
+		return InManagedViewport->MediaOutput == MediaOutput;
+	});
+
+	if (ManagedViewportIndex == INDEX_NONE)
+	{
+		checkf(false, TEXT("A managed viewport must have been registered prior to attempting the capture"));
+		return nullptr;
+	}
+	
+	UMediaCapture* MediaCapture = FindOrCreateMediaCaptureForIndex(InMediaOutputIndex);
+	if (!IsValid(MediaCapture))
+	{
+		return nullptr;
+	}
+
+	// Check to see that the media capture is not already capturing
+	EMediaCaptureState CaptureState = MediaCapture->GetState();
+	if (CaptureState == EMediaCaptureState::Capturing || CaptureState == EMediaCaptureState::Preparing)
+	{
+		return MediaCapture;
+	};
+	
+	if (MediaCapture->CaptureSceneViewport(ManagedOutputViewports[ManagedViewportIndex]->Viewport, InCaptureOptions))
+	{
+#if WITH_EDITOR
+		if (bAutoRestartCapture)
+		{
+			FRecaptureFunc RestartCaptureFunc = [InMediaOutputIndex, InCaptureOptions, bAutoRestartCapture](UMediaProfilePlaybackManager* InPlaybackManager)
+			{
+				InPlaybackManager->OpenManagedViewportOutputFromIndex(InMediaOutputIndex, InCaptureOptions, bAutoRestartCapture);
+			};
+			
+			MediaOutput->OnOutputModified().AddUObject(this, &UMediaProfilePlaybackManager::QueueRestartCapture, RestartCaptureFunc);
+		}
+#endif
+		
+		return MediaCapture;
+	}
+
+	return nullptr;
+}
+
+UMediaCapture* UMediaProfilePlaybackManager::OpenManagedViewportOutput(UMediaOutput* InMediaOutput, const FMediaCaptureOptions& InCaptureOptions, bool bAutoRestartCapture)
+{
+	UMediaProfile* MediaProfile = GetOwningMediaProfile();
+	if (!MediaProfile)
+	{
+		return nullptr;
+	}
+	
+	int32 MediaOutputIndex = MediaProfile->FindMediaOutputIndex(InMediaOutput);
+	return OpenManagedViewportOutputFromIndex(MediaOutputIndex, InCaptureOptions, bAutoRestartCapture);
+}
+
+TSharedPtr<FSceneViewport> UMediaProfilePlaybackManager::GetActiveViewportFromIndex(int32 InMediaOutputIndex) const
+{
+	using namespace MediaProfilePlaybackManager;
+
+	UMediaProfile* MediaProfile = GetOwningMediaProfile();
+	if (!MediaProfile)
+	{
+		return nullptr;
+	}
+
+	UMediaOutput* MediaOutput = MediaProfile->GetMediaOutput(InMediaOutputIndex);
+	if (!IsValid(MediaOutput))
+	{
+		return nullptr;
+	}
+	
+	// First, look for any stored viewports that are actively being captured for the media output, and return that if one is found
+	int32 ActiveViewportIndex = ActiveOutputViewports.IndexOfByPredicate([MediaOutput](const TSharedPtr<FActiveViewport>& InViewport)
+	{
+		return InViewport->MediaOutput == MediaOutput;
+	});
+
+	if (ActiveViewportIndex != INDEX_NONE)
+	{
+		return ActiveOutputViewports[ActiveViewportIndex]->GetSceneViewport();
+	}
+
+	// Otherwise, query the engine and find an active viewport to return
+	if (TSharedPtr<FSceneViewport> ActiveViewport = FindActiveViewport())
+	{
+		return ActiveViewport;
+	}
+			
+#if WITH_EDITOR
+	if (TSharedPtr<IAssetViewport> ActiveViewportInterface = FindActiveViewportInterface())
+	{
+		return ActiveViewportInterface->GetSharedActiveViewport();
+	}
+#endif
+
+	return nullptr;
+}
+
+TSharedPtr<FSceneViewport> UMediaProfilePlaybackManager::GetActiveViewport(UMediaOutput* InMediaOutput) const
+{
+	UMediaProfile* MediaProfile = GetOwningMediaProfile();
+	if (!MediaProfile)
+	{
+		return nullptr;
+	}
+	
+	int32 MediaOutputIndex = MediaProfile->FindMediaOutputIndex(InMediaOutput);
+	return GetActiveViewportFromIndex(MediaOutputIndex);
+}
+
+UMediaCapture* UMediaProfilePlaybackManager::OpenActiveViewportOutputFromIndex(int32 InMediaOutputIndex, const FMediaCaptureOptions& InCaptureOptions, bool bAutoRestartCapture)
+{
+	using namespace MediaProfilePlaybackManager;
+	
+	UMediaProfile* MediaProfile = GetOwningMediaProfile();
+	if (!MediaProfile)
+	{
+		return nullptr;
+	}
+
+	UMediaOutput* MediaOutput = MediaProfile->GetMediaOutput(InMediaOutputIndex);
+	if (!IsValid(MediaOutput))
+	{
+		return nullptr;
+	}
+	
+	UMediaCapture* MediaCapture = FindOrCreateMediaCaptureForIndex(InMediaOutputIndex);
+	if (!IsValid(MediaCapture))
+	{
+		return nullptr;
+	}
+	
+	// Check to see that the media capture is not already capturing
+	EMediaCaptureState CaptureState = MediaCapture->GetState();
+	if (CaptureState != EMediaCaptureState::Stopped && CaptureState != EMediaCaptureState::StopRequested && CaptureState != EMediaCaptureState::Error)
+	{
+		return MediaCapture;
+	};
+
+	int32 ActiveViewportIndex = ActiveOutputViewports.IndexOfByPredicate([MediaOutput](const TSharedPtr<FActiveViewport>& InViewport)
+	{
+		return InViewport->MediaOutput == MediaOutput;
+	});
+
+	if (ActiveViewportIndex == INDEX_NONE)
+	{
+#if WITH_EDITOR
+		if (ActiveOutputViewports.IsEmpty())
+		{
+			if (GEditor)
+			{
+				GEditor->OnLevelViewportClientListChanged().AddUObject(this, &UMediaProfilePlaybackManager::OnLevelViewportClientListChanged);
+			}
+		}
+#endif
+		
+		ActiveViewportIndex = ActiveOutputViewports.Add(MakeShared<FActiveViewport>(MediaOutput, InCaptureOptions));
+	}
+
+	TSharedPtr<FSceneViewport> ActiveViewport = ActiveOutputViewports[ActiveViewportIndex]->GetSceneViewport();
+	if (MediaCapture->CaptureSceneViewport(ActiveViewport, InCaptureOptions))
+	{
+#if WITH_EDITOR
+		if (bAutoRestartCapture)
+		{
+			FRecaptureFunc RestartCaptureFunc = [InMediaOutputIndex, InCaptureOptions, bAutoRestartCapture](UMediaProfilePlaybackManager* InPlaybackManager)
+			{
+				InPlaybackManager->OpenActiveViewportOutputFromIndex(InMediaOutputIndex, InCaptureOptions, bAutoRestartCapture);
+			};
+			
+			MediaOutput->OnOutputModified().AddUObject(this, &UMediaProfilePlaybackManager::QueueRestartCapture, RestartCaptureFunc);
+		}
+#endif
+		
+		return MediaCapture;
+	}
+
+	return nullptr;
+}
+
+UMediaCapture* UMediaProfilePlaybackManager::OpenActiveViewportOutput(UMediaOutput* InMediaOutput, const FMediaCaptureOptions& InCaptureOptions, bool bAutoRestartCapture)
+{
+	UMediaProfile* MediaProfile = GetOwningMediaProfile();
+	if (!MediaProfile)
+	{
+		return nullptr;
+	}
+	
+	int32 MediaOutputIndex = MediaProfile->FindMediaOutputIndex(InMediaOutput);
+	return OpenActiveViewportOutputFromIndex(MediaOutputIndex, InCaptureOptions, bAutoRestartCapture);
+}
+
+UMediaCapture* UMediaProfilePlaybackManager::OpenRenderTargetOutputFromIndex(int32 InMediaOutputIndex, UTextureRenderTarget2D* InRenderTarget, const FMediaCaptureOptions& InCaptureOptions, bool bAutoRestartCapture)
+{
+	UMediaProfile* MediaProfile = GetOwningMediaProfile();
+	if (!MediaProfile)
+	{
+		return nullptr;
+	}
+
+	UMediaOutput* MediaOutput = MediaProfile->GetMediaOutput(InMediaOutputIndex);
+	if (!IsValid(MediaOutput))
+	{
+		return nullptr;
+	}
+	
+	UMediaCapture* MediaCapture = FindOrCreateMediaCaptureForIndex(InMediaOutputIndex);
+	if (!IsValid(MediaCapture))
+	{
+		return nullptr;
+	}
+	
+	// Check to see that the media capture is not already capturing
+	EMediaCaptureState CaptureState = MediaCapture->GetState();
+	if (CaptureState != EMediaCaptureState::Stopped && CaptureState != EMediaCaptureState::StopRequested && CaptureState != EMediaCaptureState::Error)
+	{
+		return MediaCapture;
+	};
+
+	if (MediaCapture->CaptureTextureRenderTarget2D(InRenderTarget, InCaptureOptions))
+	{
+#if WITH_EDITOR
+		if (bAutoRestartCapture)
+		{
+			FRecaptureFunc RestartCaptureFunc = [InMediaOutputIndex, InRenderTarget, InCaptureOptions, bAutoRestartCapture](UMediaProfilePlaybackManager* InPlaybackManager)
+			{
+				InPlaybackManager->OpenRenderTargetOutputFromIndex(InMediaOutputIndex, InRenderTarget, InCaptureOptions, bAutoRestartCapture);
+			};
+			
+			MediaOutput->OnOutputModified().AddUObject(this, &UMediaProfilePlaybackManager::QueueRestartCapture, RestartCaptureFunc);
+		}
+#endif
+		
+		return MediaCapture;
+	}
+
+	return nullptr;
+}
+
+UMediaCapture* UMediaProfilePlaybackManager::OpenRenderTargetOutput(UMediaOutput* InMediaOutput, UTextureRenderTarget2D* InRenderTarget, const FMediaCaptureOptions& InCaptureOptions, bool bAutoRestartCapture)
+{
+	UMediaProfile* MediaProfile = GetOwningMediaProfile();
+	if (!MediaProfile)
+	{
+		return nullptr;
+	}
+	
+	int32 MediaOutputIndex = MediaProfile->FindMediaOutputIndex(InMediaOutput);
+	return OpenRenderTargetOutputFromIndex(MediaOutputIndex, InRenderTarget, InCaptureOptions, bAutoRestartCapture);
+}
+
+bool UMediaProfilePlaybackManager::IsOutputCapturingFromIndex(int32 InMediaOutputIndex) const
+{
+	UMediaProfile* MediaProfile = GetOwningMediaProfile();
+	if (!MediaProfile)
+	{
+		return false;
+	}
+
+	UMediaOutput* MediaOutput = MediaProfile->GetMediaOutput(InMediaOutputIndex);
+	if (!IsValid(MediaOutput))
+	{
+		return false;
+	}
+	
+	if (!MediaOutputCaptures.Contains(InMediaOutputIndex))
+	{
+		return false;
+	}
+	
+	UMediaCapture* MediaCapture = MediaOutputCaptures[InMediaOutputIndex];
+	if (MediaCapture->GetState() == EMediaCaptureState::Stopped || MediaCapture->GetState() == EMediaCaptureState::Error)
+	{
+		return false;
+	}
+	
+	return MediaCapture->GetMediaOutputName() == MediaOutput->GetName();
+}
+
+bool UMediaProfilePlaybackManager::IsOutputCapturing(UMediaOutput* InMediaOutput) const
+{
+	UMediaProfile* MediaProfile = GetOwningMediaProfile();
+	if (!MediaProfile)
+	{
+		return false;
+	}
+
+	int32 MediaOutputIndex = MediaProfile->FindMediaOutputIndex(InMediaOutput);
+	return IsOutputCapturingFromIndex(MediaOutputIndex);
+}
+
+TOptional<EMediaCaptureState> UMediaProfilePlaybackManager::GetOutputCaptureStateFromIndex(int32 InMediaOutputIndex, bool& bOutHasError) const
+{
+	UMediaProfile* MediaProfile = GetOwningMediaProfile();
+	if (!MediaProfile)
+	{
+		return TOptional<EMediaCaptureState>();
+	}
+
+	if (!MediaOutputCaptures.Contains(InMediaOutputIndex))
+	{
+		return TOptional<EMediaCaptureState>();
+	}
+
+	UMediaCapture* MediaCapture = MediaOutputCaptures[InMediaOutputIndex];
+
+	bOutHasError = CapturesWithError.Contains(MediaCapture);
+	return MediaCapture->GetState();
+}
+
+TOptional<EMediaCaptureState> UMediaProfilePlaybackManager::GetOutputCaptureState(UMediaOutput* InMediaOutput, bool& bOutHasError) const
+{
+	UMediaProfile* MediaProfile = GetOwningMediaProfile();
+	if (!MediaProfile)
+	{
+		return TOptional<EMediaCaptureState>();
+	}
+
+	int32 MediaOutputIndex = MediaProfile->FindMediaOutputIndex(InMediaOutput);
+	return GetOutputCaptureStateFromIndex(MediaOutputIndex, bOutHasError);
+}
+
+void UMediaProfilePlaybackManager::CloseOutputFromIndex(int32 InMediaOutputIndex, bool bDestroyCaptureObjects)
+{
+	using namespace MediaProfilePlaybackManager;
+	
+	UMediaProfile* MediaProfile = GetOwningMediaProfile();
+	if (!MediaProfile)
+	{
+		return;
+	}
+
+	UMediaOutput* MediaOutput = MediaProfile->GetMediaOutput(InMediaOutputIndex);
+	if (!IsValid(MediaOutput))
+	{
+		return;
+	}
+	
+	if (!MediaOutputCaptures.Contains(InMediaOutputIndex))
+	{
+		return;
+	}
+
+	UMediaCapture* MediaCapture = MediaOutputCaptures[InMediaOutputIndex];
+
+	constexpr bool bAllowPendingFrameToBeProcessed = false;
+	MediaCapture->StopCapture(bAllowPendingFrameToBeProcessed);
+
+	if (bDestroyCaptureObjects)
+	{
+		ManagedOutputViewports.RemoveAll([MediaOutput](const TSharedPtr<FManagedViewport>& InManagedViewport)
+		{
+			return InManagedViewport->MediaOutput == MediaOutput && InManagedViewport->Consumers.IsEmpty();
+		});
+
+		ActiveOutputViewports.RemoveAll([MediaOutput](const TSharedPtr<FActiveViewport>& InViewport)
+		{
+			return InViewport->MediaOutput == MediaOutput;
+		});
+
+#if WITH_EDITOR
+		if (ManagedOutputViewports.IsEmpty())
+		{
+			FEditorDelegates::PostPIEStarted.RemoveAll(this);
+			FEditorDelegates::PrePIEEnded.RemoveAll(this);
+		}
+
+		if (ActiveOutputViewports.IsEmpty())
+		{
+			if (GEditor)
+			{
+				GEditor->OnLevelViewportClientListChanged().RemoveAll(this);
+			}
+		}
+#endif
+		
+		MediaCapture->OnStateChangedNative.RemoveAll(this);
+		MediaOutputCaptures.Remove(InMediaOutputIndex);
+	}
+
+#if WITH_EDITOR
+	MediaOutput->OnOutputModified().RemoveAll(this);
+#endif
+}
+
+void UMediaProfilePlaybackManager::CloseOutput(UMediaOutput* InMediaOutput, bool bDestroyCaptureObjects)
+{
+	UMediaProfile* MediaProfile = GetOwningMediaProfile();
+	if (!MediaProfile)
+	{
+		return;
+	}
+	
+	int32 MediaOutputIndex = MediaProfile->FindMediaOutputIndex(InMediaOutput);
+	CloseOutputFromIndex(MediaOutputIndex, bDestroyCaptureObjects);
+}
+
 UMediaPlayer* UMediaProfilePlaybackManager::FindOrCreateMediaPlayerForIndex(int32 InMediaSourceIndex)
 {
 	UMediaProfile* MediaProfile = GetOwningMediaProfile();
@@ -399,3 +1271,224 @@
 
 	return 0;
 }
+
+UMediaCapture* UMediaProfilePlaybackManager::FindOrCreateMediaCaptureForIndex(int32 InMediaOutputIndex)
+{
+	UMediaProfile* MediaProfile = GetOwningMediaProfile();
+	if (!MediaProfile)
+	{
+		return nullptr;
+	}
+	
+	UMediaOutput* MediaOutput = MediaProfile->GetMediaOutput(InMediaOutputIndex);
+	if (!IsValid(MediaOutput))
+	{
+		return nullptr;
+	}
+
+	if (!MediaOutputCaptures.Contains(InMediaOutputIndex))
+	{
+		UMediaCapture* NewMediaCapture = MediaOutput->CreateMediaCapture();
+		if (!IsValid(NewMediaCapture))
+		{
+			return nullptr;
+		}
+
+		NewMediaCapture->OnStateChangedNative.AddUObject(this, &UMediaProfilePlaybackManager::OnCaptureStateChanged, NewMediaCapture);
+		
+		MediaOutputCaptures.Add(InMediaOutputIndex, NewMediaCapture);
+	}
+
+	return MediaOutputCaptures[InMediaOutputIndex];
+}
+
+void UMediaProfilePlaybackManager::ChangeMediaOutputIndex(int32 InOriginalMediaOutputIndex, int32 InNewMediaOutputIndex)
+{
+	UMediaProfile* MediaProfile = GetOwningMediaProfile();
+	if (!MediaProfile)
+	{
+		return;
+	}
+
+	if (InOriginalMediaOutputIndex == InNewMediaOutputIndex)
+	{
+		return;
+	}
+
+	// First, remove the existing media player from the list of media players and shift all slots below it up one
+	UMediaCapture* MediaCapture = nullptr;
+	if (InOriginalMediaOutputIndex != INDEX_NONE)
+	{
+		MediaCapture = MediaOutputCaptures.Contains(InOriginalMediaOutputIndex) ? MediaOutputCaptures[InOriginalMediaOutputIndex] : nullptr;
+		int32 LastIndex = INDEX_NONE;
+		for (int32 Index = InOriginalMediaOutputIndex + 1; Index < MediaProfile->NumMediaOutputs(); ++Index)
+		{
+			if (MediaOutputCaptures.Contains(Index))
+			{
+				MediaOutputCaptures.Add(Index - 1, MediaOutputCaptures[Index]);
+				LastIndex = Index;
+			}
+		}
+
+		if (LastIndex != INDEX_NONE)
+		{
+			MediaOutputCaptures.Remove(LastIndex);
+		}
+	}
+
+	// Second, insert the media player back into the list, shifting all slots at or below that index down one
+	if (InNewMediaOutputIndex != INDEX_NONE)
+	{
+		for (int32 Index = MediaProfile->NumMediaOutputs() - 1; Index >= InNewMediaOutputIndex; --Index)
+		{
+			if (MediaOutputCaptures.Contains(Index))
+			{
+				MediaOutputCaptures.Add(Index + 1, MediaOutputCaptures[Index]);
+			}
+		}
+
+		if (MediaCapture)
+		{
+			MediaOutputCaptures.Add(InNewMediaOutputIndex, MediaCapture);
+		}
+	}
+}
+
+void UMediaProfilePlaybackManager::OnCaptureStateChanged(UMediaCapture* InMediaCapture)
+{
+	if (InMediaCapture->GetState() == EMediaCaptureState::Error)
+	{
+		CapturesWithError.Add(InMediaCapture);
+	}
+	else if (InMediaCapture->GetState() == EMediaCaptureState::Preparing)
+	{
+		CapturesWithError.Remove(InMediaCapture);
+	}
+}
+
+void UMediaProfilePlaybackManager::QueueRestartCapture(UMediaOutput* InMediaOutput, FRecaptureFunc InRestartCaptureFunc)
+{
+	constexpr bool bDestroyCaptureObjects = false;
+	CloseOutput(InMediaOutput, bDestroyCaptureObjects);
+	
+	RestartCaptureFuncs.Add(InRestartCaptureFunc);
+	
+	if (!RestartCapturesTimerHandle.IsValid())
+	{
+#if WITH_EDITOR
+		if (GEditor)
+		{
+			RestartCapturesTimerHandle = GEditor->GetTimerManager()->SetTimerForNextTick(FTimerDelegate::CreateUObject(this, &UMediaProfilePlaybackManager::RestartCaptures));
+		}
+		else
+#endif
+		if (UWorld* World = GetWorld())
+		{
+			RestartCapturesTimerHandle = World->GetTimerManager().SetTimerForNextTick(FTimerDelegate::CreateUObject(this, &UMediaProfilePlaybackManager::RestartCaptures));
+		}
+	}
+}
+
+void UMediaProfilePlaybackManager::RestartCaptures()
+{
+	for (FRecaptureFunc& RestartCaptureFunc : RestartCaptureFuncs)
+	{
+		RestartCaptureFunc(this);
+	}
+
+	RestartCaptureFuncs.Empty();
+	RestartCapturesTimerHandle.Invalidate();
+}
+
+void UMediaProfilePlaybackManager::Cleanup()
+{
+	TArray<int32> OpenSources;
+	MediaSourcePlayers.GetKeys(OpenSources);
+	for (int32 MediaSourceIndex : OpenSources)
+	{
+		FCloseSourceArgs Args;
+		Args.bForceClose = true;
+		Args.bDestroyMediaPlayer = true;
+		
+		CloseSourceFromIndex(MediaSourceIndex, Args);
+	}
+
+	TArray<int32> OpenOutputs;
+	MediaOutputCaptures.GetKeys(OpenOutputs);
+	for (int32 MediaOutputIndex : OpenOutputs)
+	{
+		CloseOutputFromIndex(MediaOutputIndex);
+	}
+
+	ManagedOutputViewports.Empty();
+	ActiveOutputViewports.Empty();
+
+#if WITH_EDITOR
+	FEditorDelegates::PostPIEStarted.RemoveAll(this);
+	FEditorDelegates::PrePIEEnded.RemoveAll(this);
+	
+	if (GEditor)
+	{
+		GEditor->OnLevelViewportClientListChanged().RemoveAll(this);
+	}
+#endif
+}
+
+#if WITH_EDITOR
+void UMediaProfilePlaybackManager::OnPostPIEStarted(bool bIsSimulating)
+{
+	using namespace MediaProfilePlaybackManager;
+	
+	for (const TSharedPtr<FManagedViewport>& ManagedViewport : ManagedOutputViewports)
+	{
+		if (TSharedPtr<FMediaCaptureEditorViewportClient> ViewportClient = StaticCastSharedPtr<FMediaCaptureEditorViewportClient>(ManagedViewport->ViewportClient))
+		{
+			ViewportClient->SetPIE(true);
+		}
+	}
+}
+
+void UMediaProfilePlaybackManager::OnPrePIEEnded(bool bIsSimulating)
+{
+	using namespace MediaProfilePlaybackManager;
+	
+	for (const TSharedPtr<FManagedViewport>& ManagedViewport : ManagedOutputViewports)
+	{
+		if (TSharedPtr<FMediaCaptureEditorViewportClient> ViewportClient = StaticCastSharedPtr<FMediaCaptureEditorViewportClient>(ManagedViewport->ViewportClient))
+		{
+			ViewportClient->SetPIE(false);
+		}
+	}
+}
+
+void UMediaProfilePlaybackManager::OnLevelViewportClientListChanged()
+{
+	using namespace MediaProfilePlaybackManager;
+
+	if (!GEditor)
+	{
+		return;
+	}
+	
+	for (const TSharedPtr<FActiveViewport>& ActiveViewport : ActiveOutputViewports)
+	{
+		bool bStillActive = false;
+		if (TSharedPtr<FSceneViewport> Viewport = ActiveViewport->GetSceneViewport())
+		{
+			for (FLevelEditorViewportClient* ViewportClient : GEditor->GetLevelViewportClients())
+			{
+				if (ViewportClient->Viewport == Viewport.Get())
+				{
+					bStillActive = true;
+					break;
+				}
+			}
+		}
+
+		if (!bStillActive)
+		{
+			CloseOutput(ActiveViewport->MediaOutput.Get());
+		}
+	}
+}
+#endif
diff --git a/Engine/Plugins/Media/MediaFrameworkUtilities/Source/MediaFrameworkUtilities/Public/Profile/MediaProfilePlaybackManager.h b/Engine/Plugins/Media/MediaFrameworkUtilities/Source/MediaFrameworkUtilities/Public/Profile/MediaProfilePlaybackManager.h
--- a/Engine/Plugins/Media/MediaFrameworkUtilities/Source/MediaFrameworkUtilities/Public/Profile/MediaProfilePlaybackManager.h
+++ b/Engine/Plugins/Media/MediaFrameworkUtilities/Source/MediaFrameworkUtilities/Public/Profile/MediaProfilePlaybackManager.h
@@ -3,13 +3,28 @@
 #pragma once
 
 #include "UObject/Object.h"
+#include "Engine/EngineBaseTypes.h"
+#include "Engine/TimerHandle.h"
 
 #include "MediaProfilePlaybackManager.generated.h"
 
+enum class EMediaCaptureState : uint8;
+struct FMediaCaptureOptions;
+class FSceneViewport;
+class FViewportClient;
+class UMediaOutput;
+class UMediaCapture;
 class UMediaSource;
 class UMediaTexture;
 class UMediaPlayer;
 class UMediaProfile;
+class UTextureRenderTarget2D;
+
+namespace MediaProfilePlaybackManager
+{
+	class FManagedViewport;
+	class FActiveViewport;
+}
 
 /** Manages playback of media items within a media profile */
 UCLASS()
@@ -17,8 +32,12 @@ class MEDIAFRAMEWORKUTILITIES_API UMediaProfilePlaybackManager : public UObject
 {
 	GENERATED_BODY()
 
+	/** Alias for callback used to restart capture of media outputs */
+	using FRecaptureFunc = TFunction<void(UMediaProfilePlaybackManager*)>;
+	
 public:
 	UMediaProfilePlaybackManager(const FObjectInitializer& ObjectInitializer);
+	virtual ~UMediaProfilePlaybackManager() override;
 	
 public:
 	/**
@@ -112,6 +131,154 @@ public:
 	 * @param InArgs Arguments that control what happens when the media source is closed
 	 */
 	void CloseSourcesForConsumer(const FCloseSourceArgs& InArgs);
+
+	/**
+	 * Finds an existing or creates a new viewport to use for capture for a specified media output
+	 * @param InMediaOutputIndex The index of the media output in the media profile
+	 * @param InViewMode The view mode to create the viewport with
+	 * @param Consumer Optional pointer to a consumer of the viewport, which ensures the viewport is not destroyed when the capture is closed
+	 * @return The viewport client of the viewport of one was found or created; otherwise, null
+	 */
+	TSharedPtr<FViewportClient> GetOrCreateManagedViewportFromIndex(int32 InMediaOutputIndex, EViewModeIndex InViewMode, void* Consumer = nullptr);
+	
+	/**
+	 * Finds an existing or creates a new viewport to use for capture for a specified media output
+	 * @param InMediaOutput The media output in the media profile
+	 * @param InViewMode The view mode to create the viewport with
+	 * @param Consumer Optional pointer to a consumer of the viewport, which ensures the viewport is not destroyed when the capture is closed
+	 * @return The viewport client of the viewport of one was found or created; otherwise, null
+	 */
+	TSharedPtr<FViewportClient> GetOrCreateManagedViewport(UMediaOutput* InMediaOutput, EViewModeIndex InViewMode, void* Consumer = nullptr);
+
+	/**
+	 * Releases a use of the managed viewport for the specified media output,which may destroy the managed viewport if it is not actively being used
+	 * @param InMediaOutputIndex The index of the media output in the media profile
+	 * @param Consumer The pointer of the consumer that is releasing its reference of the managed viewport
+	 */
+	void ReleaseManagedViewportFromIndex(int32 InMediaOutputIndex, void* Consumer);
+
+	/**
+	 * Releases a use of the managed viewport for the specified media output,which may destroy the managed viewport if it is not actively being used
+	 * @param InMediaOutput The media output in the media profile
+	 * @param Consumer The pointer of the consumer that is releasing its reference of the managed viewport
+	 */
+	void ReleaseManagedViewport(UMediaOutput* InMediaOutput, void* Consumer);
+
+	/**
+	 * Starts a capture of a managed viewport for a specified media output
+	 * @param InMediaOutputIndex The index of the media output in the media profile
+	 * @param InCaptureOptions The capture options to capture with
+	 * @param bAutoRestartCapture Indicates that this capture will automatically restart if any modifications to the media output are detected
+	 * @return The media capture object, if the capture was started successfully; otherwise, null
+	 */
+	UMediaCapture* OpenManagedViewportOutputFromIndex(int32 InMediaOutputIndex, const FMediaCaptureOptions& InCaptureOptions, bool bAutoRestartCapture = false);
+	
+	/**
+	 * Starts a capture of a managed viewport for a specified media output
+	 * @param InMediaOutput The media output in the media profile
+	 * @param InCaptureOptions The capture options to capture with
+	 * @param bAutoRestartCapture Indicates that this capture will automatically restart if any modifications to the media output are detected
+	 * @return The media capture object, if the capture was started successfully; otherwise, null
+	 */
+	UMediaCapture* OpenManagedViewportOutput(UMediaOutput* InMediaOutput, const FMediaCaptureOptions& InCaptureOptions, bool bAutoRestartCapture = false);
+
+	/**
+	 * Gets an existing viewport in the engine that is or can be captured for the specified media output
+	 * @param InMediaOutputIndex The index of the media output in the media profile
+	 * @return A reference to the existing viewport, if one was found; otherwise, null
+	 */
+	TSharedPtr<FSceneViewport> GetActiveViewportFromIndex(int32 InMediaOutputIndex) const;
+	
+	/**
+	 * Gets an existing viewport in the engine that is or can be captured for the specified media output
+	 * @param InMediaOutput The media output in the media profile
+	 * @return A reference to the existing viewport, if one was found; otherwise, null
+	 */
+	TSharedPtr<FSceneViewport> GetActiveViewport(UMediaOutput* InMediaOutput) const;
+
+	/**
+	 * Starts a capture of an existing engine viewport for a specified media output
+	 * @param InMediaOutputIndex The index of the media output in the media profile
+	 * @param InCaptureOptions The capture options to capture with
+	 * @param bAutoRestartCapture Indicates that this capture will automatically restart if any modifications to the media output are detected
+	 * @return The media capture object, if the capture was started successfully; otherwise, null
+	 */
+	UMediaCapture* OpenActiveViewportOutputFromIndex(int32 InMediaOutputIndex, const FMediaCaptureOptions& InCaptureOptions, bool bAutoRestartCapture = false);
+
+	/**
+	 * Starts a capture of an existing engine viewport for a specified media output
+	 * @param InMediaOutput The media output in the media profile
+	 * @param InCaptureOptions The capture options to capture with
+	 * @param bAutoRestartCapture Indicates that this capture will automatically restart if any modifications to the media output are detected
+	 * @return The media capture object, if the capture was started successfully; otherwise, null
+	 */
+	UMediaCapture* OpenActiveViewportOutput(UMediaOutput* InMediaOutput, const FMediaCaptureOptions& InCaptureOptions, bool bAutoRestartCapture = false);
+
+	/**
+	 * Starts a capture of a render target for a specified media output
+	 * @param InMediaOutputIndex The index of the media output in the media profile
+	 * @param InRenderTarget The render target to capture
+	 * @param InCaptureOptions The capture options to capture with
+	 * @param bAutoRestartCapture Indicates that this capture will automatically restart if any modifications to the media output are detected
+	 * @return The media capture object, if the capture was started successfully; otherwise, null
+	 */
+	UMediaCapture* OpenRenderTargetOutputFromIndex(int32 InMediaOutputIndex, UTextureRenderTarget2D* InRenderTarget, const FMediaCaptureOptions& InCaptureOptions, bool bAutoRestartCapture = false);
+	
+	/**
+	 * Starts a capture of a render target for a specified media output
+	 * @param InMediaOutput The media output in the media profile
+	 * @param InRenderTarget The render target to capture
+	 * @param InCaptureOptions The capture options to capture with
+	 * @param bAutoRestartCapture Indicates that this capture will automatically restart if any modifications to the media output are detected
+	 * @return The media capture object, if the capture was started successfully; otherwise, null
+	 */
+	UMediaCapture* OpenRenderTargetOutput(UMediaOutput* InMediaOutput, UTextureRenderTarget2D* InRenderTarget, const FMediaCaptureOptions& InCaptureOptions, bool bAutoRestartCapture = false);
+	
+	/**
+	 * Gets whether there is an active capture for the specified media output index
+	 * @param InMediaOutputIndex The index of the media output within the media profile to check
+	 * @return true if the media output is actively being captured; otherwise false
+	 */
+	bool IsOutputCapturingFromIndex(int32 InMediaOutputIndex) const;
+	
+	/**
+	 * Gets whether there is an active capture for the specified media output
+	 * @param InMediaOutput The media output within the media profile to check
+	 * @return true if the media output is actively being captured; otherwise false
+	 */
+	bool IsOutputCapturing(UMediaOutput* InMediaOutput) const;
+
+	/**
+	 * Gets the capture state of the specified media output, or null of there is no active capture of the media output
+	 * @param InMediaOutputIndex The index of the media output within the media profile to check
+	 * @param bOutHasError Indicates that the capture state experienced an error at some point and is stopped
+	 * @return The capture state of the media output, or a null TOptional if there is no active capture of the media output
+	 */
+	TOptional<EMediaCaptureState> GetOutputCaptureStateFromIndex(int32 InMediaOutputIndex, bool& bOutHasError) const;
+
+	/**
+	 * Gets the capture state of the specified media output, or null of there is no active capture of the media output
+	 * @param InMediaOutput The media output within the media profile to check
+	 * @param bOutHasError Indicates that the capture state experienced an error at some point and is stopped
+	 * @return The capture state of the media output, or a null TOptional if there is no active capture of the media output
+	 */
+	TOptional<EMediaCaptureState> GetOutputCaptureState(UMediaOutput* InMediaOutput, bool& bOutHasError) const;
+	
+	/**
+	 * Closes the capture of the specified media output
+	 * @param InMediaOutputIndex The index of the media output in the media profile to close
+	 * @param bDestroyCaptureObjects Indicates that the media capture and all related objects should be destroyed.
+	 *								 Set to false if the capture is immediately going to be restarted.
+	 */
+	void CloseOutputFromIndex(int32 InMediaOutputIndex, bool bDestroyCaptureObjects = true);
+
+	/**
+	 * Closes the capture of the specified media output
+	 * @param InMediaOutput The media output in the media profile to close
+	 * @param bDestroyCaptureObjects Indicates that the media capture and all related objects should be destroyed.
+	 *								 Set to false if the capture is immediately going to be restarted.
+	 */
+	void CloseOutput(UMediaOutput* InMediaOutput, bool bDestroyCaptureObjects = true);
 	
 private:
 	/**
@@ -139,6 +306,39 @@ private:
 
 	/** Gets the number of consumers that have registered as using the specified media player */
 	int32 GetMediaSourceConsumerCount(UMediaPlayer* InMediaPlayer) const;
+
+	/** Attempts to find an existing valid capture or creates a new capture for the specified media output index */
+	UMediaCapture* FindOrCreateMediaCaptureForIndex(int32 InMediaOutputIndex);
+
+	/**
+	 * Processes a change to a media output's slot index, moving any existing media captures around
+	 * @param InOriginalMediaOutputIndex The original index of the media output, or INDEX_NONE if we are processing an insert
+	 * @param InNewMediaOutputIndex The new index of the media output, or INDEX_NONE if we are processing a delete
+	 */
+	void ChangeMediaOutputIndex(int32 InOriginalMediaOutputIndex, int32 InNewMediaOutputIndex);
+
+	/** Raised when the capture state of the specified media capture has changed. Used to keep track of which captures may have errored */
+	void OnCaptureStateChanged(UMediaCapture* InMediaCapture);
+
+	/** Raised when a change to the specified media output has been detected and the capture shout be auto-restarted */
+	void QueueRestartCapture(UMediaOutput* InMediaOutput, FRecaptureFunc InRestartCaptureFunc);
+
+	/** Called on the next available tick, restarts any captures that have been queued to auto-restart */
+	void RestartCaptures();
+
+	/** Cleans up all open sources and outputs, and deletes any managed state */
+	void Cleanup();
+	
+#if WITH_EDITOR
+	/** Raised when a PIE session has started */
+	void OnPostPIEStarted(bool bIsSimulating);
+
+	/** Raised when a PIE session is ending */
+	void OnPrePIEEnded(bool bIsSimulating);
+
+	/** Raised when the level editor's list of viewports has changed */
+	void OnLevelViewportClientListChanged();
+#endif
 	
 private:
 	/** A map of media players, where the hash key is the index of the media source the player is playing */
@@ -149,8 +349,27 @@ private:
 	UPROPERTY(Transient)
 	TArray<TSoftObjectPtr<UMediaTexture>> MediaSourceTextures;
 
+	/** A map of media captures, where the hash key is the index of the media output being captured */
+	UPROPERTY(Transient)
+	TMap<int32, TObjectPtr<UMediaCapture>> MediaOutputCaptures;
+	
 	/** Keeps track of registered consumers that are actively using an open media source */
 	TMap<TObjectPtr<UMediaPlayer>, TSet<void*>> SourceConsumers;
 	
+	/** A list of viewports being managed by the playback manager,  */
+	TArray<TSharedPtr<MediaProfilePlaybackManager::FManagedViewport>> ManagedOutputViewports;
+
+	/** A list of active viewports that are being captured from */
+	TArray<TSharedPtr<MediaProfilePlaybackManager::FActiveViewport>> ActiveOutputViewports;
+
+	/** Keeps track of which captures have reported an error with the capture */
+	TSet<TObjectPtr<UMediaCapture>> CapturesWithError;
+	
+	/** A list of callbacks to use to auto-restart any stopped captures on the next tick */
+	TArray<FRecaptureFunc> RestartCaptureFuncs;
+
+	/** Timer handle for auto-restarting output captures */
+	FTimerHandle RestartCapturesTimerHandle;
+	
 	friend class UMediaProfile;
 };
diff --git a/Engine/Plugins/Media/MediaFrameworkUtilities/Source/MediaFrameworkUtilitiesEditor/Private/AssetEditor/MediaProfileEditor.cpp b/Engine/Plugins/Media/MediaFrameworkUtilities/Source/MediaFrameworkUtilitiesEditor/Private/AssetEditor/MediaProfileEditor.cpp
--- a/Engine/Plugins/Media/MediaFrameworkUtilities/Source/MediaFrameworkUtilitiesEditor/Private/AssetEditor/MediaProfileEditor.cpp
+++ b/Engine/Plugins/Media/MediaFrameworkUtilities/Source/MediaFrameworkUtilitiesEditor/Private/AssetEditor/MediaProfileEditor.cpp
@@ -52,162 +52,6 @@ namespace MediaProfileEditor
 	}
 }
 
-/*
- * FMediaFrameworkCaptureLevelEditorViewportClient
- * Like FLevelEditorViewportClient but always use the PlayWorld if available. Do not support bSimulated
- */
-class FMediaFrameworkCaptureLevelEditorViewportClient : public FLevelEditorViewportClient
-{
-private:
-	using Super = FLevelEditorViewportClient;
-
-public:
-	static TSharedPtr<FMediaFrameworkCaptureLevelEditorViewportClient> CreateViewportClientForCapture(const FMediaFrameworkCaptureCameraViewportCameraOutputInfo& InViewportOutputInfo)
-	{
-		TSharedPtr<FMediaFrameworkCaptureLevelEditorViewportClient> NewViewportClient = MakeShared<FMediaFrameworkCaptureLevelEditorViewportClient>(nullptr, InViewportOutputInfo.ViewMode);
-
-		// Preview viewports never be a listener
-		NewViewportClient->bSetListenerPosition = false;
-
-		// Default to "game" show flags for camera previews
-		NewViewportClient->EngineShowFlags = FEngineShowFlags(ESFIM_Game);
-		NewViewportClient->LastEngineShowFlags = FEngineShowFlags(ESFIM_Editor);
-
-		NewViewportClient->ViewportType = LVT_Perspective;
-		NewViewportClient->bDrawAxes = false;
-		NewViewportClient->bDisableInput = true;
-		NewViewportClient->SetAllowCinematicControl(false);
-		NewViewportClient->VisibilityDelegate.BindLambda([] {return true; });
-		NewViewportClient->AutoSetPIE();
-		
-		return NewViewportClient;
-	}
-	
-	FMediaFrameworkCaptureLevelEditorViewportClient(const TSharedPtr<SLevelViewport>& InLevelViewport, EViewModeIndex InViewModeIndex)
-		: FLevelEditorViewportClient(InLevelViewport)
-		, ViewModeIndex(InViewModeIndex)
-	{}
-
-	virtual ~FMediaFrameworkCaptureLevelEditorViewportClient() override
-	{
-		if (GEditor)
-		{
-			FWorldContext* PIEWorldContext = GEditor->GetPIEWorldContext();
-			if (PIEWorldContext)
-			{
-				RemoveReferenceToWorldContext(*PIEWorldContext);
-			}
-		}
-	}
-
-	virtual UWorld* GetWorld() const override
-	{
-		if (GEditor->PlayWorld)
-		{
-			return GEditor->PlayWorld;
-		}
-		return Super::GetWorld();
-	}
-
-	void SetPIE(bool bInIsPIE)
-	{
-		FWorldContext* PIEWorldContext = GEditor->GetPIEWorldContext();
-		if (bInIsPIE)
-		{
-			if (PIEWorldContext)
-			{
-				RemoveReferenceToWorldContext(GEditor->GetEditorWorldContext());
-				SetReferenceToWorldContext(*PIEWorldContext);
-			}
-		}
-		else
-		{
-			if (PIEWorldContext)
-			{
-				RemoveReferenceToWorldContext(*PIEWorldContext);
-			}
-			SetReferenceToWorldContext(GEditor->GetEditorWorldContext());
-		}
-		SetViewMode(ViewModeIndex);
-		SetRealtime(true);
-	}
-
-	void AutoSetPIE()
-	{
-		if (GEditor->PlayWorld)
-		{
-			SetPIE(true);
-		}
-		else
-		{
-			SetViewMode(ViewModeIndex);
-			SetRealtime(true);
-		}
-	}
-
-private:
-	EViewModeIndex ViewModeIndex;
-};
-
-FMediaProfileEditor::FViewportCapture::~FViewportCapture()
-{
-	ViewportClient->Viewport = nullptr;
-
-	ViewportClient.Reset();
-	Viewport.Reset();
-	ViewportWidget.Reset();
-}
-
-FMediaProfileEditor::FActiveEditorViewport::FActiveEditorViewport(const TSharedPtr<FSceneViewport>& InViewport, const TSharedPtr<IAssetViewport>& InViewportInterface)
-	: Viewport(InViewport)
-	, ViewportInterface(InViewportInterface)
-{
-	if (InViewportInterface.IsValid())
-	{
-		FEditorViewportClient& ViewportClient = InViewportInterface->GetAssetViewportClient();
-
-		// Cache the viewport's current settings, then set them to more appropriate settings for recording
-		CachedViewportFlags.bRealTime = ViewportClient.IsRealtime();
-		CachedViewportFlags.bSetListenerPosition = ViewportClient.bSetListenerPosition;
-		CachedViewportFlags.bDrawAxes = ViewportClient.bDrawAxes;
-		CachedViewportFlags.bDisableInput = ViewportClient.bDisableInput;
-		CachedViewportFlags.bAllowCinematicControl = ViewportClient.AllowsCinematicControl();
-		CachedViewportFlags.VisibilityDelegate = ViewportClient.VisibilityDelegate;
-
-		ViewportClient.SetRealtime(true);
-		ViewportClient.bSetListenerPosition = false;
-		ViewportClient.bDrawAxes = false;
-		ViewportClient.SetAllowCinematicControl(false);
-		ViewportClient.VisibilityDelegate.BindLambda([] { return true; });
-		
-		if (UMediaFrameworkWorldSettingsAssetUserData* CaptureSettings = FindOrAddMediaFrameworkCaptureSettings())
-		{
-			// Disable input by default unless we're dynamically resizing the viewport texture as we capture.
-			// Note that we originally disable input because click routing is broken when we resize the viewport through media capture at the moment.
-			const bool bDisableInput = CaptureSettings->CurrentViewportMediaOutput.CaptureOptions.ResizeMethod != EMediaCaptureResizeMethod::ResizeInRenderPass;
-			ViewportClient.bDisableInput = bDisableInput ? true : ViewportClient.bDisableInput;
-		}
-	}
-}
-
-FMediaProfileEditor::FActiveEditorViewport::~FActiveEditorViewport()
-{
-	TSharedPtr<FSceneViewport> PinnedViewport = Viewport.Pin();
-	TSharedPtr<IAssetViewport> PinnedViewportInterface = ViewportInterface.Pin();
-	if (PinnedViewportInterface.IsValid() && PinnedViewportInterface->GetSharedActiveViewport() == PinnedViewport)
-	{
-		FEditorViewportClient& ViewportClient = PinnedViewportInterface->GetAssetViewportClient();
-
-		// Reset settings
-		ViewportClient.SetRealtime(CachedViewportFlags.bRealTime);
-		ViewportClient.bSetListenerPosition = CachedViewportFlags.bSetListenerPosition;
-		ViewportClient.bDrawAxes = CachedViewportFlags.bDrawAxes;
-		ViewportClient.bDisableInput = CachedViewportFlags.bDisableInput;
-		ViewportClient.SetAllowCinematicControl(CachedViewportFlags.bAllowCinematicControl);
-		ViewportClient.VisibilityDelegate = CachedViewportFlags.VisibilityDelegate;
-	}
-}
-
 TSharedRef<FMediaProfileEditor> FMediaProfileEditor::CreateMediaProfileEditor(const EToolkitMode::Type Mode, const TSharedPtr<IToolkitHost>& InitToolkitHost, UMediaProfile* InMediaProfile)
 {
 	TSharedRef<FMediaProfileEditor> NewEditor = MakeShared<FMediaProfileEditor>();
@@ -220,15 +64,12 @@ void FMediaProfileEditor::Initialize(const EToolkitMode::Type Mode, const TShare
 	GEditor->GetEditorSubsystem<UAssetEditorSubsystem>()->CloseOtherEditors(InMediaProfile, this);
 	MediaProfileBeingEdited = InMediaProfile;
 
-	FEditorDelegates::PostPIEStarted.AddSP(this, &FMediaProfileEditor::OnPostPIEStarted);
-	FEditorDelegates::PrePIEEnded.AddSP(this, &FMediaProfileEditor::OnPrePIEEnded);
 	FEditorDelegates::MapChange.AddSP(this, &FMediaProfileEditor::OnMapChange);
 	GEngine->OnLevelActorDeleted().AddSP(this, &FMediaProfileEditor::OnLevelActorsRemoved);
 	FEditorDelegates::OnAssetsDeleted.AddSP(this, &FMediaProfileEditor::OnAssetsDeleted);
 	FCoreUObjectDelegates::OnPreObjectPropertyChanged.AddSP(this, &FMediaProfileEditor::OnObjectPreEditChange);
 	FCoreUObjectDelegates::OnObjectModified.AddSP(this, &FMediaProfileEditor::OnObjectModified);
 	FCoreUObjectDelegates::OnObjectPropertyChanged.AddSP(this, &FMediaProfileEditor::OnObjectPropertyChanged);
-	GEditor->OnLevelViewportClientListChanged().AddSP(this, &FMediaProfileEditor::OnLevelViewportClientListChanged);
 	
 	const TSharedRef<FTabManager::FLayout> StandaloneDefaultLayout = FTabManager::NewLayout("Standalone_MediaProfileEditor_Layout")
 		->AddArea
@@ -295,7 +136,6 @@ FMediaProfileEditor::~FMediaProfileEditor()
 	FEditorDelegates::PostPIEStarted.RemoveAll(this);
 
 	CloseAllMediaSources();
-	StopCapturingAllMediaOutputs();
 }
 
 void FMediaProfileEditor::RegisterTabSpawners(const TSharedRef<class FTabManager>& InTabManager)
@@ -365,6 +205,11 @@ FLinearColor FMediaProfileEditor::GetWorldCentricTabColorScale() const
 
 void FMediaProfileEditor::CloseAllMediaSources()
 {
+	if (!IsValid(MediaProfileBeingEdited))
+	{
+		return;
+	}
+	
 	for (int32 Index = 0; Index < MediaProfileBeingEdited->NumMediaSources(); ++Index)
 	{
 		UMediaProfilePlaybackManager::FCloseSourceArgs Args;
@@ -374,73 +219,26 @@ void FMediaProfileEditor::CloseAllMediaSources()
 	}
 }
 
-TSharedPtr<FSceneViewport> FMediaProfileEditor::GetActiveEditorViewport() const
+void FMediaProfileEditor::CloseAllMediaOutputs()
 {
-	// If there is a live capture of the active editor viewport, return the viewport it is using
-	for (const FLiveMediaCapture& Capture : LiveCaptures)
-	{
-		if (Capture.ActiveEditorViewport.IsValid())
-		{
-			return Capture.ActiveEditorViewport->Viewport.Pin();
-		}
-	}
-
-	// Otherwise, attempt to find a viewport
-	TSharedPtr<FSceneViewport> Viewport;
-	TSharedPtr<IAssetViewport> ViewportInterface;
-	GetActiveEditorViewportImpl(Viewport, ViewportInterface);
-	return Viewport;
-}
-
-TSharedPtr<FMediaProfileEditor::FViewportCapture> FMediaProfileEditor::FindOrCreateViewportCapture(UMediaOutput* InMediaOutput)
-{
-	UMediaFrameworkWorldSettingsAssetUserData* CaptureSettings = FindOrAddMediaFrameworkCaptureSettings();
-	if (!CaptureSettings)
-	{
-		return nullptr;
-	}
-
-	if (!InMediaOutput)
+	if (!IsValid(MediaProfileBeingEdited))
 	{
-		return nullptr;
+		return;
 	}
-
-	int32 ViewportOutputInfoIndex = CaptureSettings->ViewportCaptures.IndexOfByPredicate([InMediaOutput](const FMediaFrameworkCaptureCameraViewportCameraOutputInfo& ViewportOutputInfo)
-	{
-		return ViewportOutputInfo.MediaOutput == InMediaOutput;
-	});
-
-	if (ViewportOutputInfoIndex == INDEX_NONE)
+	
+	for (int32 MediaOutputIndex = 0; MediaOutputIndex < MediaProfileBeingEdited->NumMediaOutputs(); ++MediaOutputIndex)
 	{
-		return nullptr;
+		MediaProfileBeingEdited->GetPlaybackManager()->CloseOutputFromIndex(MediaOutputIndex);
 	}
-
-	return FindOrCreateViewportCaptureImpl(ViewportOutputInfoIndex);
 }
 
-bool FMediaProfileEditor::IsMediaOutputCapturing(UMediaOutput* InMediaOutput) const
-{
-	return LiveCaptures.ContainsByPredicate([InMediaOutput](const FLiveMediaCapture& LiveCapture)
-	{
-		return LiveCapture.MediaOutput.IsValid() && LiveCapture.MediaOutput.Get() == InMediaOutput;
-	});
-}
-
-TOptional<EMediaCaptureState> FMediaProfileEditor::GetMediaOutputCaptureState(UMediaOutput* InMediaOutput) const
+bool FMediaProfileEditor::CanMediaOutputCapture(UMediaOutput* InMediaOutput) const
 {
-	for (const FLiveMediaCapture& Capture : LiveCaptures)
+	if (!IsValid(MediaProfileBeingEdited))
 	{
-		if (Capture.MediaOutput == InMediaOutput)
-		{
-			return Capture.CaptureState;
-		}
+		return false;
 	}
-
-	return TOptional<EMediaCaptureState>();
-}
-
-bool FMediaProfileEditor::CanMediaOutputCapture(UMediaOutput* InMediaOutput) const
-{
+	
 	UMediaFrameworkWorldSettingsAssetUserData* CaptureSettings = FindOrAddMediaFrameworkCaptureSettings();
 	if (!CaptureSettings)
 	{
@@ -452,10 +250,7 @@ bool FMediaProfileEditor::CanMediaOutputCapture(UMediaOutput* InMediaOutput) con
 	if (CaptureSettings->CurrentViewportMediaOutput.MediaOutput == InMediaOutput)
 	{
 		//Check to see if there is a valid editor viewport that can be captured from
-		TSharedPtr<FSceneViewport> Viewport;
-		TSharedPtr<IAssetViewport> ViewportInterface;
-		GetActiveEditorViewportImpl(Viewport, ViewportInterface);
-
+		TSharedPtr<FSceneViewport> Viewport = MediaProfileBeingEdited->GetPlaybackManager()->GetActiveViewport(InMediaOutput);
 		bHasValidCaptureSettings = Viewport.IsValid();
 	}
 	
@@ -495,140 +290,6 @@ bool FMediaProfileEditor::CanMediaOutputCapture(UMediaOutput* InMediaOutput) con
 	return bHasValidCaptureSettings;
 }
 
-void FMediaProfileEditor::StartCapturingMediaOutput(UMediaOutput* InMediaOutput)
-{
-	UMediaFrameworkWorldSettingsAssetUserData* CaptureSettings = FindOrAddMediaFrameworkCaptureSettings();
-	if (!CaptureSettings)
-	{
-		return;
-	}
-
-	if (!InMediaOutput)
-	{
-		return;
-	}
-	
-	bool bStartedCapture = false;
-	if (CaptureSettings->CurrentViewportMediaOutput.MediaOutput == InMediaOutput)
-	{
-		FLiveMediaCapture NewLiveCapture;
-		NewLiveCapture.MediaOutput = InMediaOutput;
-		NewLiveCapture.MediaCapture = TStrongObjectPtr(InMediaOutput->CreateMediaCapture());
-
-		if (NewLiveCapture.MediaCapture.IsValid())
-		{
-			TSharedPtr<FSceneViewport> Viewport;
-			TSharedPtr<IAssetViewport> ViewportInterface;
-			GetActiveEditorViewportImpl(Viewport, ViewportInterface);
-
-			if (Viewport.IsValid())
-			{
-				NewLiveCapture.ActiveEditorViewport = MakeShared<FActiveEditorViewport>(Viewport, ViewportInterface);
-				if (NewLiveCapture.MediaCapture->CaptureSceneViewport(Viewport, CaptureSettings->CurrentViewportMediaOutput.CaptureOptions))
-				{
-					NewLiveCapture.CaptureState = NewLiveCapture.MediaCapture->GetState();
-					NewLiveCapture.MediaCapture->OnStateChangedNative.AddSP(this, &FMediaProfileEditor::OnMediaCaptureStateChanged, NewLiveCapture.MediaCapture.Get());
-					LiveCaptures.Add(NewLiveCapture);
-					bStartedCapture = true;
-				}
-			}
-		}
-	}
-	
-	for (int32 Index = 0; Index < CaptureSettings->ViewportCaptures.Num(); ++Index)
-	{
-		const FMediaFrameworkCaptureCameraViewportCameraOutputInfo& ViewportOutputInfo = CaptureSettings->ViewportCaptures[Index];
-		if (ViewportOutputInfo.MediaOutput != InMediaOutput)
-		{
-			continue;
-		}
-		
-		FLiveMediaCapture NewLiveCapture;
-		NewLiveCapture.MediaOutput = InMediaOutput;
-		NewLiveCapture.MediaCapture = TStrongObjectPtr(InMediaOutput->CreateMediaCapture());
-
-		if (!NewLiveCapture.MediaCapture.IsValid())
-		{
-			continue;
-		}
-		
-		NewLiveCapture.ViewportCapture = FindOrCreateViewportCaptureImpl(Index);
-		if (!NewLiveCapture.ViewportCapture.IsValid())
-		{
-			continue;
-		}
-
-		if (NewLiveCapture.MediaCapture->CaptureSceneViewport(NewLiveCapture.ViewportCapture->Viewport, ViewportOutputInfo.CaptureOptions))
-		{
-			NewLiveCapture.CaptureState = NewLiveCapture.MediaCapture->GetState();
-			NewLiveCapture.MediaCapture->OnStateChangedNative.AddSP(this, &FMediaProfileEditor::OnMediaCaptureStateChanged, NewLiveCapture.MediaCapture.Get());
-			LiveCaptures.Add(NewLiveCapture);
-			bStartedCapture = true;
-		}
-	}
-
-	for (const FMediaFrameworkCaptureRenderTargetCameraOutputInfo& RenderTargetOutputInfo : CaptureSettings->RenderTargetCaptures)
-	{
-		if (RenderTargetOutputInfo.MediaOutput != InMediaOutput)
-		{
-			continue;
-		}
-		
-		FLiveMediaCapture NewLiveCapture;
-		NewLiveCapture.MediaOutput = InMediaOutput;
-		NewLiveCapture.MediaCapture = TStrongObjectPtr(InMediaOutput->CreateMediaCapture());
-
-		if (NewLiveCapture.MediaCapture.IsValid())
-		{
-			if (NewLiveCapture.MediaCapture->CaptureTextureRenderTarget2D(RenderTargetOutputInfo.RenderTarget.Get(), RenderTargetOutputInfo.CaptureOptions))
-			{
-				NewLiveCapture.CaptureState = NewLiveCapture.MediaCapture->GetState();
-				NewLiveCapture.MediaCapture->OnStateChangedNative.AddSP(this, &FMediaProfileEditor::OnMediaCaptureStateChanged, NewLiveCapture.MediaCapture.Get());
-				LiveCaptures.Add(NewLiveCapture);
-				bStartedCapture = true;
-			}
-		}
-	}
-
-	if (bStartedCapture)
-	{
-		InMediaOutput->OnOutputModified().AddSP(this, &FMediaProfileEditor::OnMediaOutputModified);
-	}
-}
-
-void FMediaProfileEditor::StopCapturingMediaOutput(UMediaOutput* InMediaOutput)
-{
-	if (!InMediaOutput)
-	{
-		return;
-	}
-	
-	for (FLiveMediaCapture& Capture : LiveCaptures)
-	{
-		if (Capture.MediaOutput.Get() != InMediaOutput)
-		{
-			continue;
-		}
-
-		CloseLiveCapture(Capture);
-	}
-
-	LiveCaptures.RemoveAll([InMediaOutput](const FLiveMediaCapture& Capture)
-	{
-		return Capture.MediaOutput.Get() == InMediaOutput;
-	});
-}
-
-void FMediaProfileEditor::StopCapturingAllMediaOutputs()
-{
-	for (FLiveMediaCapture& Capture : LiveCaptures)
-	{
-		CloseLiveCapture(Capture);
-	}
-
-	LiveCaptures.Empty();
-}
-
 UMediaFrameworkWorldSettingsAssetUserData* FMediaProfileEditor::FindOrAddMediaFrameworkCaptureSettings()
 {
 	UMediaFrameworkWorldSettingsAssetUserData* Result = nullptr;
@@ -724,30 +385,58 @@ void FMediaProfileEditor::RefreshSelectedMediaItems(const TArray<int32>& InMedia
 		return;
 	}
 
+	UMediaProfilePlaybackManager* PlaybackManager = MediaProfileBeingEdited->GetPlaybackManager();
+	
 	for (int32 Index = 0; Index < InMediaSources.Num(); ++Index)
 	{
 		const int32 SrcIndex = InMediaSources[Index];
 
-		if (MediaProfileBeingEdited->GetPlaybackManager()->IsSourceOpenFromIndex(SrcIndex))
+		if (PlaybackManager->IsSourceOpenFromIndex(SrcIndex))
 		{
 			UMediaProfilePlaybackManager::FCloseSourceArgs Args;
 			Args.Consumer = this;
 			Args.bForceClose = true;
 		
-			MediaProfileBeingEdited->GetPlaybackManager()->CloseSourceFromIndex(SrcIndex, Args);
-			MediaProfileBeingEdited->GetPlaybackManager()->OpenSourceFromIndex(SrcIndex, this);
+			PlaybackManager->CloseSourceFromIndex(SrcIndex, Args);
+			PlaybackManager->OpenSourceFromIndex(SrcIndex, this);
 		}
 	}
 
+	const bool bAutoRestartCapture = GetDefault<UMediaFrameworkEditorCaptureSettings>()->bAutoRestartCaptureOnChange;
+	
 	for (int32 Index = 0; Index < InMediaOutputs.Num(); ++Index)
 	{
 		const int32 OutputIndex = InMediaOutputs[Index];
 		if (UMediaOutput* MediaOutput = MediaProfileBeingEdited->GetMediaOutput(OutputIndex))
 		{
-			if (IsMediaOutputCapturing(MediaOutput))
+			if (PlaybackManager->IsOutputCapturing(MediaOutput))
 			{
-				StopCapturingMediaOutput(MediaOutput);
-				StartCapturingMediaOutput(MediaOutput);
+				constexpr bool bDestroyCaptureObjects = false;
+				PlaybackManager->CloseOutput(MediaOutput, bDestroyCaptureObjects);
+
+				if (UMediaFrameworkWorldSettingsAssetUserData* CaptureSettings = FindOrAddMediaFrameworkCaptureSettings())
+				{
+					if (CaptureSettings->CurrentViewportMediaOutput.MediaOutput == MediaOutput)
+					{
+						PlaybackManager->OpenActiveViewportOutput(MediaOutput, CaptureSettings->CurrentViewportMediaOutput.CaptureOptions, bAutoRestartCapture);
+					}
+
+					for (const FMediaFrameworkCaptureCameraViewportCameraOutputInfo& ViewportCapture : CaptureSettings->ViewportCaptures)
+					{
+						if (ViewportCapture.MediaOutput == MediaOutput)
+						{
+							PlaybackManager->OpenManagedViewportOutput(MediaOutput, ViewportCapture.CaptureOptions, bAutoRestartCapture);
+						}
+					}
+
+					for (const FMediaFrameworkCaptureRenderTargetCameraOutputInfo& RenderTargetCapture : CaptureSettings->RenderTargetCaptures)
+					{
+						if (RenderTargetCapture.MediaOutput == MediaOutput)
+						{
+							PlaybackManager->OpenRenderTargetOutput(MediaOutput, RenderTargetCapture.RenderTarget, RenderTargetCapture.CaptureOptions, bAutoRestartCapture);
+						}
+					}
+				}
 			}
 		}
 	}
@@ -766,7 +455,7 @@ void FMediaProfileEditor::OnMediaItemDeleted(UClass* InMediaType, int32 InMediaI
 	}
 	else if (InMediaType->IsChildOf<UMediaOutput>())
 	{
-		StopCapturingMediaOutput(MediaProfileBeingEdited->GetMediaOutput(InMediaItemIndex));
+		MediaProfileBeingEdited->GetPlaybackManager()->CloseOutputFromIndex(InMediaItemIndex);
 	}
 
 	if (ViewportPanel.IsValid())
@@ -786,124 +475,22 @@ void FMediaProfileEditor::OnSelectedMediaItemsChanged(const TArray<int32>& Selec
 	ViewportPanel->SetSelectedMediaItems(SelectedMediaSources, SelectedMediaOutputs);
 }
 
-void FMediaProfileEditor::OnMediaOutputModified(UMediaOutput* InMediaOutput)
+void FMediaProfileEditor::OnMapChange(uint32 InMapFlags)
 {
-	if (!GetDefault<UMediaFrameworkEditorCaptureSettings>()->bAutoRestartCaptureOnChange)
-	{
-		return;
-	}
-	
-	const bool bIsCapturing = LiveCaptures.ContainsByPredicate([InMediaOutput](const FLiveMediaCapture& Capture)
-	{
-		return Capture.MediaOutput.Get() == InMediaOutput;
-	});
-
-	if (!bIsCapturing)
+	// New map might have different capture settings, so stop all captures and notify of settings changes
+	if (!IsValid(MediaProfileBeingEdited))
 	{
 		return;
 	}
 
-	StopCapturingMediaOutput(InMediaOutput);
-	OutputsToRestartCapture.Add(InMediaOutput);
+	CloseAllMediaOutputs();
 	
-	if (!RestartCapturesTimerHandle.IsValid())
-	{
-		RestartCapturesTimerHandle = GEditor->GetTimerManager()->SetTimerForNextTick(FTimerDelegate::CreateSP(this, &FMediaProfileEditor::OnRestartCapturesNextTick));
-	}
-}
-
-void FMediaProfileEditor::OnMediaCaptureStateChanged(UMediaCapture* InMediaCapture)
-{
-	if (!IsValid(InMediaCapture))
-	{
-		return;
-	}
-	
-	int32 Index = LiveCaptures.IndexOfByPredicate([InMediaCapture](const FLiveMediaCapture& InLiveCapture)
-	{
-		return InLiveCapture.MediaCapture.Get() == InMediaCapture;
-	});
-	
-	if (Index == INDEX_NONE)
-	{
-		return;
-	}
-
-	// UMediaCapture automatically transitions from Error to Stopped, but we want to cache the error state
-	// and not actually close the capture yet so that the user can be notified that the capture had an error
-	EMediaCaptureState CaptureState = InMediaCapture->GetState();
-	if (CaptureState == EMediaCaptureState::Stopped)
-	{
-		if (LiveCaptures[Index].CaptureState != EMediaCaptureState::Error)
-		{
-			CloseLiveCapture(LiveCaptures[Index]);
-			LiveCaptures.RemoveAt(Index);
-		}
-	}
-	else
-	{
-		LiveCaptures[Index].CaptureState = CaptureState;
-	}
-}
-
-void FMediaProfileEditor::OnRestartCapturesNextTick()
-{
-	for (const TWeakObjectPtr<UMediaOutput>& MediaOutput : OutputsToRestartCapture)
-	{
-		if (MediaOutput.IsValid())
-		{
-			StartCapturingMediaOutput(MediaOutput.Get());
-		}
-	}
-
-	OutputsToRestartCapture.Empty();
-	RestartCapturesTimerHandle.Invalidate();
-}
-
-void FMediaProfileEditor::OnPostPIEStarted(bool bIsSimulating)
-{
-	for (FLiveMediaCapture& Capture : LiveCaptures)
-	{
-		if (!Capture.ViewportCapture.IsValid())
-		{
-			continue;
-		}
-
-		TSharedPtr<FMediaFrameworkCaptureLevelEditorViewportClient> ViewportClient = StaticCastSharedPtr<FMediaFrameworkCaptureLevelEditorViewportClient>(Capture.ViewportCapture->ViewportClient);
-		if (ViewportClient.IsValid())
-		{
-			ViewportClient->SetPIE(true);
-		}
-	}
-}
-
-void FMediaProfileEditor::OnPrePIEEnded(bool bIsSimulating)
-{
-	for (FLiveMediaCapture& Capture : LiveCaptures)
-	{
-		if (!Capture.ViewportCapture.IsValid())
-		{
-			continue;
-		}
-
-		TSharedPtr<FMediaFrameworkCaptureLevelEditorViewportClient> ViewportClient = StaticCastSharedPtr<FMediaFrameworkCaptureLevelEditorViewportClient>(Capture.ViewportCapture->ViewportClient);
-		if (ViewportClient.IsValid())
-		{
-			ViewportClient->SetPIE(false);
-		}
-	}
-}
-
-void FMediaProfileEditor::OnMapChange(uint32 InMapFlags)
-{
-	// New map might have different capture settings, so stop all captures and notify of settings changes
-	StopCapturingAllMediaOutputs();
 	OnCaptureSettingsObjectChanged.Broadcast();
 }
 
 void FMediaProfileEditor::OnLevelActorsRemoved(AActor* InActor)
 {
-	if (!LiveCaptures.Num())
+	if (!IsValid(MediaProfileBeingEdited))
 	{
 		return;
 	}
@@ -919,14 +506,14 @@ void FMediaProfileEditor::OnLevelActorsRemoved(AActor* InActor)
 	{
 		if (ViewportOutputInfo.LockedActors.Contains(InActor))
 		{
-			StopCapturingMediaOutput(ViewportOutputInfo.MediaOutput);
+			MediaProfileBeingEdited->GetPlaybackManager()->CloseOutput(ViewportOutputInfo.MediaOutput);
 		}
 	}
 }
 
 void FMediaProfileEditor::OnAssetsDeleted(const TArray<UClass*>& DeletedAssetClasses)
 {
-	if (!LiveCaptures.Num())
+	if (!IsValid(MediaProfileBeingEdited))
 	{
 		return;
 	}
@@ -947,18 +534,23 @@ void FMediaProfileEditor::OnAssetsDeleted(const TArray<UClass*>& DeletedAssetCla
 		// If the deleted asset is a render target, we need to stop all render target captures in case the deleted render target is used in one of them
 		for (const FMediaFrameworkCaptureRenderTargetCameraOutputInfo& RenderTargetOutputInfo : CaptureSettings->RenderTargetCaptures)
 		{
-			StopCapturingMediaOutput(RenderTargetOutputInfo.MediaOutput);
+			MediaProfileBeingEdited->GetPlaybackManager()->CloseOutput(RenderTargetOutputInfo.MediaOutput);
 		}
 	}
 }
 
 void FMediaProfileEditor::OnObjectPreEditChange(UObject* Object, const FEditPropertyChain& PropertyChain)
 {
+	if (!IsValid(MediaProfileBeingEdited))
+	{
+		return;
+	}
+	
 	// If the capture settings are about to be changed, stop all captures
 	UMediaFrameworkWorldSettingsAssetUserData* CaptureSettings = FindOrAddMediaFrameworkCaptureSettings();
 	if (Object == CaptureSettings)
 	{
-		StopCapturingAllMediaOutputs();
+		CloseAllMediaOutputs();
 	}
 }
 
@@ -995,174 +587,63 @@ void FMediaProfileEditor::OnObjectPropertyChanged(UObject* InObject, FPropertyCh
 
 	if (MediaProfileBeingEdited != InObject)
 	{
+		UMediaProfilePlaybackManager* PlaybackManager = MediaProfileBeingEdited->GetPlaybackManager();
+		
 		if (UMediaSource* MediaSource = Cast<UMediaSource>(InObject))
 		{
 			const int32 SrcIndex = MediaProfileBeingEdited->FindMediaSourceIndex(MediaSource);
 			if (SrcIndex != INDEX_NONE)
 			{
-				if (MediaProfileBeingEdited->GetPlaybackManager()->IsSourceOpenFromIndex(SrcIndex))
+				if (PlaybackManager->IsSourceOpenFromIndex(SrcIndex))
 				{
 					UMediaProfilePlaybackManager::FCloseSourceArgs Args;
 					Args.Consumer = this;
 					Args.bForceClose = true;
 		
-					MediaProfileBeingEdited->GetPlaybackManager()->CloseSourceFromIndex(SrcIndex, Args);
-					MediaProfileBeingEdited->GetPlaybackManager()->OpenSourceFromIndex(SrcIndex, this);
+					PlaybackManager->CloseSourceFromIndex(SrcIndex, Args);
+					PlaybackManager->OpenSourceFromIndex(SrcIndex, this);
 				}
 			}
 		}
 
 		if (UMediaOutput* MediaOutput = Cast<UMediaOutput>(InObject))
 		{
+			const bool bAutoRestartCapture = GetDefault<UMediaFrameworkEditorCaptureSettings>()->bAutoRestartCaptureOnChange;
 			const int32 OutputIndex = MediaProfileBeingEdited->FindMediaOutputIndex(MediaOutput);
 			if (OutputIndex != INDEX_NONE)
 			{
-				if (IsMediaOutputCapturing(MediaOutput))
+				if (PlaybackManager->IsOutputCapturing(MediaOutput))
 				{
-					StopCapturingMediaOutput(MediaOutput);
-					StartCapturingMediaOutput(MediaOutput);
-				}
-			}
-		}
-	}	
-}
+					constexpr bool bDestroyCaptureObjects = false;
+					PlaybackManager->CloseOutput(MediaOutput, bDestroyCaptureObjects);
 
-void FMediaProfileEditor::OnLevelViewportClientListChanged()
-{
-	for (int32 Index = 0; Index < LiveCaptures.Num(); ++Index)
-	{
-		FLiveMediaCapture& Capture = LiveCaptures[Index];
-		if (Capture.ActiveEditorViewport.IsValid())
-		{
-			bool bStillActive = false;
-			if (TSharedPtr<FSceneViewport> PinnedViewport = Capture.ActiveEditorViewport->Viewport.Pin())
-			{
-				for (FLevelEditorViewportClient* ViewportClient : GEditor->GetLevelViewportClients())
-				{
-					if (ViewportClient->Viewport == PinnedViewport.Get())
+					if (UMediaFrameworkWorldSettingsAssetUserData* CaptureSettings = FindOrAddMediaFrameworkCaptureSettings())
 					{
-						bStillActive = true;
-						break;
+						if (CaptureSettings->CurrentViewportMediaOutput.MediaOutput == MediaOutput)
+						{
+							PlaybackManager->OpenActiveViewportOutput(MediaOutput, CaptureSettings->CurrentViewportMediaOutput.CaptureOptions, bAutoRestartCapture);
+						}
+
+						for (const FMediaFrameworkCaptureCameraViewportCameraOutputInfo& ViewportCapture : CaptureSettings->ViewportCaptures)
+						{
+							if (ViewportCapture.MediaOutput == MediaOutput)
+							{
+								PlaybackManager->OpenManagedViewportOutput(MediaOutput, ViewportCapture.CaptureOptions, bAutoRestartCapture);
+							}
+						}
+
+						for (const FMediaFrameworkCaptureRenderTargetCameraOutputInfo& RenderTargetCapture : CaptureSettings->RenderTargetCaptures)
+						{
+							if (RenderTargetCapture.MediaOutput == MediaOutput)
+							{
+								PlaybackManager->OpenRenderTargetOutput(MediaOutput, RenderTargetCapture.RenderTarget, RenderTargetCapture.CaptureOptions, bAutoRestartCapture);
+							}
+						}
 					}
 				}
 			}
-
-			if (!bStillActive)
-			{
-				CloseLiveCapture(Capture);
-				LiveCaptures.RemoveAt(Index--);
-			}
-		}
-	}
-}
-
-TSharedPtr<FMediaProfileEditor::FViewportCapture> FMediaProfileEditor::FindOrCreateViewportCaptureImpl(int32 ViewportOutputInfoIndex)
-{
-	// Take the opportunity to clear any invalid viewport captures from the list
-	TArray<int32> ActiveViewportCaptureKeys;
-	ActiveViewportCaptures.GetKeys(ActiveViewportCaptureKeys);
-	for (int32 ViewportCaptureKey : ActiveViewportCaptureKeys)
-	{
-		if (!ActiveViewportCaptures[ViewportCaptureKey].IsValid())
-		{
-			ActiveViewportCaptures.Remove(ViewportCaptureKey);
-		}
-	}
-	
-	UMediaFrameworkWorldSettingsAssetUserData* CaptureSettings = FindOrAddMediaFrameworkCaptureSettings();
-	if (!CaptureSettings)
-	{
-		return nullptr;
-	}
-
-	if (!CaptureSettings->ViewportCaptures.IsValidIndex(ViewportOutputInfoIndex))
-	{
-		return nullptr;
-	}
-	
-	const FMediaFrameworkCaptureCameraViewportCameraOutputInfo& ViewportOutputInfo = CaptureSettings->ViewportCaptures[ViewportOutputInfoIndex];
-	if (!ActiveViewportCaptures.Contains(ViewportOutputInfoIndex))
-	{
-		if (!ViewportOutputInfo.MediaOutput) 
-		{
-			return nullptr;
-		}
-
-		TSharedPtr<FViewportCapture> NewViewportCapture = MakeShared<FViewportCapture>();
-		NewViewportCapture->MediaOutput = ViewportOutputInfo.MediaOutput;
-		NewViewportCapture->ViewportClient = FMediaFrameworkCaptureLevelEditorViewportClient::CreateViewportClientForCapture(ViewportOutputInfo);
-
-		NewViewportCapture->ViewportWidget = SNew(SViewport)
-			.RenderDirectlyToWindow(false)
-			.IsEnabled(FSlateApplication::Get().GetNormalExecutionAttribute())
-			.EnableGammaCorrection(false)
-			.EnableBlending(false);
-		
-		NewViewportCapture->Viewport = MakeShared<FSceneViewport>(NewViewportCapture->ViewportClient.Get(), NewViewportCapture->ViewportWidget);
-		NewViewportCapture->ViewportClient->Viewport = NewViewportCapture->Viewport.Get();
-		NewViewportCapture->ViewportWidget->SetViewportInterface(NewViewportCapture->Viewport.ToSharedRef());
-		
-		ActiveViewportCaptures.Add(ViewportOutputInfoIndex, NewViewportCapture);
-
-		return NewViewportCapture;
-	}
-	
-	return ActiveViewportCaptures[ViewportOutputInfoIndex].Pin();
-}
-
-void FMediaProfileEditor::GetActiveEditorViewportImpl(TSharedPtr<FSceneViewport>& OutSceneViewport, TSharedPtr<IAssetViewport>& OutViewportInterface) const
-{
-	OutSceneViewport = nullptr;
-	OutViewportInterface = nullptr;
-	
-	UMediaFrameworkWorldSettingsAssetUserData* CaptureSettings = FindOrAddMediaFrameworkCaptureSettings();
-	if (!CaptureSettings)
-	{
-		return;
-	}
-
-	if (!IsValid(CaptureSettings->CurrentViewportMediaOutput.MediaOutput))
-	{
-		return;
-	}
-
-	// If there is an active PIE viewport, return it
-	for (const FWorldContext& Context : GEngine->GetWorldContexts())
-	{
-		if (Context.WorldType == EWorldType::PIE)
-		{
-			UEditorEngine* EditorEngine = CastChecked<UEditorEngine>(GEngine);
-			FSlatePlayInEditorInfo& Info = EditorEngine->SlatePlayInEditorMap.FindChecked(Context.ContextHandle);
-			if (Info.SlatePlayInEditorWindowViewport.IsValid())
-			{
-				OutSceneViewport = Info.SlatePlayInEditorWindowViewport;
-			}
 		}
-	}
-
-	if (OutSceneViewport.IsValid())
-	{
-		return;
-	}
-	
-	// Find a level editor viewport
-	FLevelEditorModule& LevelEditorModule = FModuleManager::Get().GetModuleChecked<FLevelEditorModule>(TEXT("LevelEditor"));
-	OutViewportInterface = LevelEditorModule.GetFirstActiveViewport();
-	OutSceneViewport = OutViewportInterface.IsValid() ? OutViewportInterface->GetSharedActiveViewport() : nullptr;
-}
-
-void FMediaProfileEditor::CloseLiveCapture(FLiveMediaCapture& InLiveCapture)
-{
-	if (InLiveCapture.MediaCapture.IsValid())
-	{
-		InLiveCapture.MediaCapture->OnStateChangedNative.RemoveAll(this);
-		InLiveCapture.MediaCapture->StopCapture(false);
-	}
-
-	if (InLiveCapture.MediaOutput.IsValid())
-	{
-		InLiveCapture.MediaOutput->OnOutputModified().RemoveAll(this);
-	}
+	}	
 }
 
 #undef LOCTEXT_NAMESPACE
diff --git a/Engine/Plugins/Media/MediaFrameworkUtilities/Source/MediaFrameworkUtilitiesEditor/Private/AssetEditor/MediaProfileEditor.h b/Engine/Plugins/Media/MediaFrameworkUtilities/Source/MediaFrameworkUtilitiesEditor/Private/AssetEditor/MediaProfileEditor.h
--- a/Engine/Plugins/Media/MediaFrameworkUtilities/Source/MediaFrameworkUtilitiesEditor/Private/AssetEditor/MediaProfileEditor.h
+++ b/Engine/Plugins/Media/MediaFrameworkUtilities/Source/MediaFrameworkUtilitiesEditor/Private/AssetEditor/MediaProfileEditor.h
@@ -4,7 +4,6 @@
 
 #include "CoreMinimal.h"
 #include "EditorViewportClient.h"
-#include "ShowFlags.h"
 #include "Toolkits/AssetEditorToolkit.h"
 
 enum class EMediaCaptureState : uint8;
@@ -30,52 +29,8 @@ private:
 	static const FName TimecodeTabId;
 
 public:
-	/** Stores objects necessary for rendering viewport camera captures to a widget */
-	struct FViewportCapture
-	{
-		TWeakObjectPtr<UMediaOutput> MediaOutput;
-		TSharedPtr<FLevelEditorViewportClient> ViewportClient;
-		TSharedPtr<FSceneViewport> Viewport;
-		TSharedPtr<SViewport> ViewportWidget;
-
-		~FViewportCapture();
-	};
-	
 	DECLARE_MULTICAST_DELEGATE_OneParam(FOnCaptureMethodChanged, UMediaOutput*)
 
-private:
-	/** Stores objects necessary for capturing the currently active editor viewport */
-	struct FActiveEditorViewport
-	{
-		TWeakPtr<FSceneViewport> Viewport;
-		TWeakPtr<IAssetViewport> ViewportInterface;
-
-		struct
-		{
-			bool bRealTime = false;
-			bool bSetListenerPosition = false;
-			bool bDrawAxes = false;
-			bool bDisableInput = false;
-			bool bAllowCinematicControl = false;
-			FEngineShowFlags EngineShowFlags = ESFIM_Editor;
-			FEngineShowFlags LastEngineShowFlags = ESFIM_Game;
-			FViewportStateGetter VisibilityDelegate;
-		} CachedViewportFlags;
-
-		FActiveEditorViewport(const TSharedPtr<FSceneViewport>& InViewport, const TSharedPtr<IAssetViewport>& InViewportInterface);
-		~FActiveEditorViewport();
-	};
-	
-	/** Stores objects necessary for performing a live capture of a media output */
-	struct FLiveMediaCapture
-	{
-		TWeakObjectPtr<UMediaOutput> MediaOutput;
-		TStrongObjectPtr<UMediaCapture> MediaCapture;
-		TSharedPtr<FViewportCapture> ViewportCapture;
-		TSharedPtr<FActiveEditorViewport> ActiveEditorViewport;
-		EMediaCaptureState CaptureState;
-	};
-	
 public:
 	static TSharedRef<FMediaProfileEditor> CreateMediaProfileEditor(const EToolkitMode::Type Mode, const TSharedPtr< IToolkitHost >& InitToolkitHost, UMediaProfile* InMediaProfile);
 
@@ -103,34 +58,12 @@ public:
 	/** Closes all open media sources in the media profile */
 	void CloseAllMediaSources();
 
-	/** Returns a valid editor viewport that can be captured from, or null if there are none */
-	TSharedPtr<FSceneViewport> GetActiveEditorViewport() const;
-	
-	/**
-	 * Attempts to find an existing viewport capture for the specified media output, or creates one if one does not exist
-	 * @param InMediaOutput The media output to find a viewport capture for
-	 * @return The viewport capture, if found or created, or null if the media output does not have a viewport capture configured
-	 */
-	TSharedPtr<FViewportCapture> FindOrCreateViewportCapture(UMediaOutput* InMediaOutput);
-	
-	/** Determines if the specified media output is currently capturing */
-	bool IsMediaOutputCapturing(UMediaOutput* InMediaOutput) const;
-
-	/** Gets the current state of the live capture for the specified media output, or a null TOptional if there is no live capture */
-	TOptional<EMediaCaptureState> GetMediaOutputCaptureState(UMediaOutput* InMediaOutput) const;
+	/** Closes all open media outputs in the media profile */
+	void CloseAllMediaOutputs();
 	
 	/** Determines if the specified media output is properly configured for capturing */
 	bool CanMediaOutputCapture(UMediaOutput* InMediaOutput) const;
-	
-	/** Starts capturing the specified media output */
-	void StartCapturingMediaOutput(UMediaOutput* InMediaOutput);
-
-	/** Stops capture on any live captures for the specified media output */
-	void StopCapturingMediaOutput(UMediaOutput* InMediaOutput);
 
-	/** Stops capture on all live captures */
-	void StopCapturingAllMediaOutputs();
-	
 	/** Gets a valid media capture settings object if one is found, or attempts to add one to the world if one is not found */
 	static UMediaFrameworkWorldSettingsAssetUserData* FindOrAddMediaFrameworkCaptureSettings();
 
@@ -153,13 +86,6 @@ private:
 	
 	void OnMediaItemDeleted(UClass* InMediaType, int32 InMediaItemIndex);
 	void OnSelectedMediaItemsChanged(const TArray<int32>& SelectedMediaSources, const TArray<int32>& SelectedMediaOutputs);
-	void OnMediaOutputModified(UMediaOutput* InMediaOutput);
-	void OnMediaCaptureStateChanged(UMediaCapture* InMediaCapture);
-	void OnRestartCapturesNextTick();
-
-
-	void OnPostPIEStarted(bool bIsSimulating);
-	void OnPrePIEEnded(bool bIsSimulating);
 	
 	void OnMapChange(uint32 InMapFlags);
 	void OnLevelActorsRemoved(AActor* InActor);
@@ -167,16 +93,6 @@ private:
 	void OnObjectPreEditChange(UObject* Object, const FEditPropertyChain& PropertyChain);
 	void OnObjectModified(UObject* Object);
 	void OnObjectPropertyChanged(UObject* InObject, FPropertyChangedEvent& InPropertyChangedEvent);
-	void OnLevelViewportClientListChanged();
-
-	/** Attempts to find an existing viewport capture for the specified viewport output info, or creates one if one does not exist */
-	TSharedPtr<FViewportCapture> FindOrCreateViewportCaptureImpl(int32 ViewportOutputInfoIndex);
-
-	/** Attempts to find an editor viewport to capture from (can be either a PIE viewport or a level editor viewport) */
-	void GetActiveEditorViewportImpl(TSharedPtr<FSceneViewport>& OutSceneViewport, TSharedPtr<IAssetViewport>& OutViewportInterface) const;
-
-	/** Closes the live media capture and performs any necessary cleanup */
-	void CloseLiveCapture(FLiveMediaCapture& InLiveCapture);
 	
 private:
 	/** The media profile being edited by this editor */
@@ -188,21 +104,6 @@ private:
 	/** Timecode/Genlock header widget that is displayed in the editor's primary toolbar */
 	TSharedPtr<SMediaFrameworkTimecodeGenlockHeader> TimecodeToolbarEntry;
 
-	/**
-	 * A map that stores the viewport captures for specific media capture camera viewport outputs.
-	 * There are two ways a viewport capture is consumed:
-	 *     When a media output is being displayed in the media profile editor viewport
-	 *     When a media output is actively being captured.
-	 * To avoid having to create duplicate media capture structs for each consumer, store active captures here
-	 * that can be consumed by either or both the viewport and the live media capture. Further, store viewport
-	 * captures using weak pointers so that when all consumers of a viewport capture release their viewport capture pointers,
-	 * the viewport captures are cleaned up, with the invalid weak pointers cleaned up at the next opportunity
-	 */
-	TMap<int32, TWeakPtr<FViewportCapture>> ActiveViewportCaptures;
-	
-	/** List of any live media output captures that the media profile editor is managing */
-	TArray<FLiveMediaCapture> LiveCaptures;
-
 	/** List of media outputs whose capture needs to be restarted on the next engine tick */
 	TSet<TWeakObjectPtr<UMediaOutput>> OutputsToRestartCapture;
 
diff --git a/Engine/Plugins/Media/MediaFrameworkUtilities/Source/MediaFrameworkUtilitiesEditor/Private/AssetEditor/MediaProfileMediaItemDetailCustomization.cpp b/Engine/Plugins/Media/MediaFrameworkUtilities/Source/MediaFrameworkUtilitiesEditor/Private/AssetEditor/MediaProfileMediaItemDetailCustomization.cpp
--- a/Engine/Plugins/Media/MediaFrameworkUtilities/Source/MediaFrameworkUtilitiesEditor/Private/AssetEditor/MediaProfileMediaItemDetailCustomization.cpp
+++ b/Engine/Plugins/Media/MediaFrameworkUtilities/Source/MediaFrameworkUtilitiesEditor/Private/AssetEditor/MediaProfileMediaItemDetailCustomization.cpp
@@ -7,6 +7,7 @@
 #include "IStructureDataProvider.h"
 #include "MediaFrameworkWorldSettingsAssetUserData.h"
 #include "Algo/Count.h"
+#include "Profile/MediaProfilePlaybackManager.h"
 #include "UI/MediaFrameworkUtilitiesEditorStyle.h"
 
 #define LOCTEXT_NAMESPACE "FMediaProfileMediaItemDetailCustomization"
@@ -639,7 +640,7 @@ void FMediaProfileMediaOutputDetailCustomization::SetMediaObject(UMediaOutput* I
 		Index = Cast<UDummyMediaOutput>(InOriginalMediaObject)->MediaProfileIndex;
 	}
 
-	PinnedMediaProfileEditor->StopCapturingMediaOutput(InOriginalMediaObject);
+	PinnedMediaProfile->GetPlaybackManager()->CloseOutputFromIndex(Index);
 	
 	UMediaFrameworkWorldSettingsAssetUserData* CaptureSettings = FMediaProfileEditor::FindOrAddMediaFrameworkCaptureSettings();
 	if (!CaptureSettings)
@@ -715,6 +716,12 @@ void FMediaProfileMediaOutputDetailCustomization::OnCaptureMethodChanged()
 		return;
 	}
 
+	TStrongObjectPtr<UMediaProfile> PinnedMediaProfile = MediaProfile.Pin();
+	if (!PinnedMediaProfile.IsValid())
+	{
+		return;
+	}
+	
 	TArray<TWeakObjectPtr<UObject>> Objects;
 	CachedDetailBuilder.Pin()->GetObjectsBeingCustomized(Objects);
 
@@ -737,7 +744,7 @@ void FMediaProfileMediaOutputDetailCustomization::OnCaptureMethodChanged()
 		}
 
 		UMediaOutput* MediaOutput = Cast<UMediaOutput>(PinnedObject.Get());
-		PinnedMediaProfileEditor->StopCapturingMediaOutput(MediaOutput);
+		PinnedMediaProfile->GetPlaybackManager()->CloseOutput(MediaOutput);
 		PinnedMediaProfileEditor->GetOnCaptureMethodChanged().Broadcast(MediaOutput);
 	}
 	
diff --git a/Engine/Plugins/Media/MediaFrameworkUtilities/Source/MediaFrameworkUtilitiesEditor/Private/AssetEditor/SMediaProfileDetailsInfoPanel.cpp b/Engine/Plugins/Media/MediaFrameworkUtilities/Source/MediaFrameworkUtilitiesEditor/Private/AssetEditor/SMediaProfileDetailsInfoPanel.cpp
--- a/Engine/Plugins/Media/MediaFrameworkUtilities/Source/MediaFrameworkUtilitiesEditor/Private/AssetEditor/SMediaProfileDetailsInfoPanel.cpp
+++ b/Engine/Plugins/Media/MediaFrameworkUtilities/Source/MediaFrameworkUtilitiesEditor/Private/AssetEditor/SMediaProfileDetailsInfoPanel.cpp
@@ -12,6 +12,7 @@
 #include "MediaAssets/ProxyMediaOutput.h"
 #include "MediaAssets/ProxyMediaSource.h"
 #include "Profile/MediaProfile.h"
+#include "Profile/MediaProfilePlaybackManager.h"
 #include "Styling/SlateIconFinder.h"
 #include "Widgets/SBoxPanel.h"
 #include "Widgets/Images/SImage.h"
@@ -487,23 +488,23 @@ FText SMediaObjectInfoPanel::GetCaptureObjectValueText() const
 FText SMediaObjectInfoPanel::GetCaptureStatusText() const
 {
 	UMediaOutput* MediaOutput = ObjectAsMediaOutput();
-	if (!MediaOutput || !MediaProfileEditor.IsValid())
+	if (!MediaOutput || !MediaProfile.IsValid())
 	{
 		return LOCTEXT("UnknownCaptureStatusText", "Unknown");
 	}
 
-	return MediaProfileEditor.Pin()->IsMediaOutputCapturing(MediaOutput) ? LOCTEXT("CapturingStatusText", "Capturing") : LOCTEXT("NotCapturingStatusText", "Not Capturing");
+	return MediaProfile.Pin()->GetPlaybackManager()->IsOutputCapturing(MediaOutput) ? LOCTEXT("CapturingStatusText", "Capturing") : LOCTEXT("NotCapturingStatusText", "Not Capturing");
 }
 
 FSlateColor SMediaObjectInfoPanel::GetCaptureStatusColor() const
 {
 	UMediaOutput* MediaOutput = ObjectAsMediaOutput();
-	if (!MediaOutput || !MediaProfileEditor.IsValid())
+	if (!MediaOutput || !MediaProfile.IsValid())
 	{
 		return FSlateColor::UseForeground();
 	}
 
-	return MediaProfileEditor.Pin()->IsMediaOutputCapturing(MediaOutput) ? FSlateColor(FLinearColor(0.0, 1.0, 0.0)) : FSlateColor::UseForeground();
+	return MediaProfile.Pin()->GetPlaybackManager()->IsOutputCapturing(MediaOutput) ? FSlateColor(FLinearColor(0.0, 1.0, 0.0)) : FSlateColor::UseForeground();
 }
 
 FText SMediaObjectInfoPanel::GetCaptureColorConversionText() const
diff --git a/Engine/Plugins/Media/MediaFrameworkUtilities/Source/MediaFrameworkUtilitiesEditor/Private/AssetEditor/SMediaProfileDetailsPanel.cpp b/Engine/Plugins/Media/MediaFrameworkUtilities/Source/MediaFrameworkUtilitiesEditor/Private/AssetEditor/SMediaProfileDetailsPanel.cpp
--- a/Engine/Plugins/Media/MediaFrameworkUtilities/Source/MediaFrameworkUtilitiesEditor/Private/AssetEditor/SMediaProfileDetailsPanel.cpp
+++ b/Engine/Plugins/Media/MediaFrameworkUtilities/Source/MediaFrameworkUtilitiesEditor/Private/AssetEditor/SMediaProfileDetailsPanel.cpp
@@ -16,6 +16,7 @@
 #include "CaptureTab/SMediaFrameworkCapture.h"
 #include "Modules/ModuleManager.h"
 #include "Profile/MediaProfile.h"
+#include "Profile/MediaProfilePlaybackManager.h"
 #include "UI/MediaFrameworkUtilitiesEditorStyle.h"
 #include "Widgets/Input/SComboButton.h"
 #include "Widgets/Input/SEditableTextBox.h"
@@ -233,12 +234,6 @@ void SMediaProfileDetailsPanel::RefreshDetailsView()
 
 void SMediaProfileDetailsPanel::OnCaptureButtonPressed()
 {
-	TSharedPtr<FMediaProfileEditor> PinnedEditor = OwningEditor.Pin();
-	if (!PinnedEditor.IsValid())
-	{
-		return;
-	}
-	
 	TStrongObjectPtr<UMediaProfile> PinnedMediaProfile = MediaProfile.Pin();
 	if (!PinnedMediaProfile.IsValid())
 	{
@@ -250,7 +245,7 @@ void SMediaProfileDetailsPanel::OnCaptureButtonPressed()
 	{
 		if (UMediaOutput* MediaOutput = PinnedMediaProfile->GetMediaOutput(OutputIndex))
 		{
-			if (PinnedEditor->IsMediaOutputCapturing(MediaOutput))
+			if (PinnedMediaProfile->GetPlaybackManager()->IsOutputCapturing(MediaOutput))
 			{
 				bIsCapturing = true;
 				break;
@@ -264,17 +259,55 @@ void SMediaProfileDetailsPanel::OnCaptureButtonPressed()
 		{
 			if (UMediaOutput* MediaOutput = PinnedMediaProfile->GetMediaOutput(OutputIndex))
 			{
-				PinnedEditor->StopCapturingMediaOutput(MediaOutput);
+				PinnedMediaProfile->GetPlaybackManager()->CloseOutput(MediaOutput);
 			}
 		}
 	}
 	else
 	{
+		UMediaFrameworkWorldSettingsAssetUserData* MediaCaptureSettings = FMediaProfileEditor::FindOrAddMediaFrameworkCaptureSettings();
+		if (!IsValid(MediaCaptureSettings))
+		{
+			return;
+		}
+		
 		for (int32 OutputIndex : SelectedMediaOutputs)
 		{
 			if (UMediaOutput* MediaOutput = PinnedMediaProfile->GetMediaOutput(OutputIndex))
 			{
-				PinnedEditor->StartCapturingMediaOutput(MediaOutput);
+				if (MediaCaptureSettings->CurrentViewportMediaOutput.MediaOutput == MediaOutput)
+				{
+					PinnedMediaProfile->GetPlaybackManager()->OpenActiveViewportOutput(MediaOutput, MediaCaptureSettings->CurrentViewportMediaOutput.CaptureOptions);
+					continue;
+				}
+
+				FMediaFrameworkCaptureCameraViewportCameraOutputInfo* ViewportCapture =
+					MediaCaptureSettings->ViewportCaptures.FindByPredicate([MediaOutput](const FMediaFrameworkCaptureCameraViewportCameraOutputInfo& InViewportCapture)
+					{
+						return InViewportCapture.MediaOutput == MediaOutput;
+					});
+				
+				if (ViewportCapture)
+				{
+					// Make sure there is a managed viewport before starting the capture
+					TSharedPtr<FViewportClient> ViewportClient = PinnedMediaProfile->GetPlaybackManager()->GetOrCreateManagedViewport(MediaOutput, ViewportCapture->ViewMode);
+					if (ViewportClient.IsValid())
+					{
+						PinnedMediaProfile->GetPlaybackManager()->OpenManagedViewportOutput(MediaOutput, ViewportCapture->CaptureOptions);
+					}
+					continue;
+				}
+
+				FMediaFrameworkCaptureRenderTargetCameraOutputInfo* RenderTargetCapture =
+					MediaCaptureSettings->RenderTargetCaptures.FindByPredicate([MediaOutput](const FMediaFrameworkCaptureRenderTargetCameraOutputInfo& InRenderTargetCapture)
+					{
+						return InRenderTargetCapture.MediaOutput == MediaOutput;
+					});
+
+				if (RenderTargetCapture)
+				{
+					PinnedMediaProfile->GetPlaybackManager()->OpenRenderTargetOutput(MediaOutput, RenderTargetCapture->RenderTarget, RenderTargetCapture->CaptureOptions);
+				}
 			}
 		}
 	}
@@ -282,12 +315,6 @@ void SMediaProfileDetailsPanel::OnCaptureButtonPressed()
 
 const FSlateBrush* SMediaProfileDetailsPanel::GetCaptureButtonImage() const
 {
-	TSharedPtr<FMediaProfileEditor> PinnedEditor = OwningEditor.Pin();
-	if (!PinnedEditor.IsValid())
-	{
-		return nullptr;
-	}
-
 	TStrongObjectPtr<UMediaProfile> PinnedMediaProfile = MediaProfile.Pin();
 	if (!PinnedMediaProfile.IsValid())
 	{
@@ -299,7 +326,7 @@ const FSlateBrush* SMediaProfileDetailsPanel::GetCaptureButtonImage() const
 	{
 		if (UMediaOutput* MediaOutput = PinnedMediaProfile->GetMediaOutput(OutputIndex))
 		{
-			if (PinnedEditor->IsMediaOutputCapturing(MediaOutput))
+			if (PinnedMediaProfile->GetPlaybackManager()->IsOutputCapturing(MediaOutput))
 			{
 				bIsCapturing = true;
 				break;
@@ -312,12 +339,6 @@ const FSlateBrush* SMediaProfileDetailsPanel::GetCaptureButtonImage() const
 
 FText SMediaProfileDetailsPanel::GetCaptureButtonText() const
 {
-	TSharedPtr<FMediaProfileEditor> PinnedEditor = OwningEditor.Pin();
-	if (!PinnedEditor.IsValid())
-	{
-		return FText::GetEmpty();
-	}
-
 	TStrongObjectPtr<UMediaProfile> PinnedMediaProfile = MediaProfile.Pin();
 	if (!PinnedMediaProfile.IsValid())
 	{
@@ -329,7 +350,7 @@ FText SMediaProfileDetailsPanel::GetCaptureButtonText() const
 	{
 		if (UMediaOutput* MediaOutput = PinnedMediaProfile->GetMediaOutput(OutputIndex))
 		{
-			if (PinnedEditor->IsMediaOutputCapturing(MediaOutput))
+			if (PinnedMediaProfile->GetPlaybackManager()->IsOutputCapturing(MediaOutput))
 			{
 				bIsCapturing = true;
 				break;
@@ -387,7 +408,7 @@ bool SMediaProfileDetailsPanel::IsCaptureButtonEnabled() const
 	{
 		if (UMediaOutput* MediaOutput = PinnedMediaProfile->GetMediaOutput(OutputIndex))
 		{
-			if (PinnedEditor->IsMediaOutputCapturing(MediaOutput))
+			if (PinnedMediaProfile->GetPlaybackManager()->IsOutputCapturing(MediaOutput))
 			{
 				bIsCapturing = true;
 			}
diff --git a/Engine/Plugins/Media/MediaFrameworkUtilities/Source/MediaFrameworkUtilitiesEditor/Private/AssetEditor/SMediaProfileMediaItemDisplay.cpp b/Engine/Plugins/Media/MediaFrameworkUtilities/Source/MediaFrameworkUtilitiesEditor/Private/AssetEditor/SMediaProfileMediaItemDisplay.cpp
--- a/Engine/Plugins/Media/MediaFrameworkUtilities/Source/MediaFrameworkUtilitiesEditor/Private/AssetEditor/SMediaProfileMediaItemDisplay.cpp
+++ b/Engine/Plugins/Media/MediaFrameworkUtilities/Source/MediaFrameworkUtilitiesEditor/Private/AssetEditor/SMediaProfileMediaItemDisplay.cpp
@@ -389,9 +389,10 @@ public:
 		SLATE_ATTRIBUTE(FOptionalSize, AspectRatio)
 	SLATE_END_ARGS()
 
-	void Construct(const FArguments& InArgs, const TWeakPtr<FMediaProfileEditor>& InMediaProfileEditor)
+	void Construct(const FArguments& InArgs, UMediaProfile* InMediaProfile, const FMediaFrameworkCaptureCurrentViewportOutputInfo& OutputInfo)
 	{
-		MediaProfileEditor = InMediaProfileEditor;
+		MediaProfile = InMediaProfile;
+		MediaOutputIndex = MediaProfile->FindMediaOutputIndex(OutputInfo.MediaOutput);
 		AspectRatio = InArgs._AspectRatio;
 		
 		SMediaCaptureImageBase::FArguments BaseArgs;
@@ -411,9 +412,9 @@ protected:
 private:
 	TWeakPtr<FSceneViewport> GetActiveEditorViewport() const
 	{
-		if (TSharedPtr<FMediaProfileEditor> PinnedMediaProfileEditor = MediaProfileEditor.Pin())
+		if (TStrongObjectPtr<UMediaProfile> PinnedMediaProfile = MediaProfile.Pin())
 		{
-			return PinnedMediaProfileEditor->GetActiveEditorViewport();
+			return PinnedMediaProfile->GetPlaybackManager()->GetActiveViewportFromIndex(MediaOutputIndex);
 		}
 
 		return nullptr;
@@ -429,9 +430,9 @@ private:
 			return AspectRatio.Get(FOptionalSize());
 		}
 
-		if (TSharedPtr<FMediaProfileEditor> PinnedMediaProfileEditor = MediaProfileEditor.Pin())
+		if (TStrongObjectPtr<UMediaProfile> PinnedMediaProfile = MediaProfile.Pin())
 		{
-			if (TSharedPtr<FSceneViewport> Viewport = PinnedMediaProfileEditor->GetActiveEditorViewport())
+			if (TSharedPtr<FSceneViewport> Viewport = PinnedMediaProfile->GetPlaybackManager()->GetActiveViewportFromIndex(MediaOutputIndex))
 			{
 				return Viewport->GetDesiredAspectRatio();
 			}
@@ -441,14 +442,14 @@ private:
 	}
 	
 private:
-	/** The media profile editor that manages the media output being displayed */
-	TWeakPtr<FMediaProfileEditor> MediaProfileEditor;
+	/** The media profile whose media output is being displayed */
+	TWeakObjectPtr<UMediaProfile> MediaProfile;
 
+	/** The index of the media output being displayed */
+	int32 MediaOutputIndex = INDEX_NONE;
+	
 	/** Desired aspect ratio of the media output being captured */
 	TAttribute<FOptionalSize> AspectRatio;
-	
-	/** Command list that manages the show flags for the viewport capture */
-	TSharedPtr<FUICommandList> ShowFlagsCommandList;
 };
 
 /** Widget to display a viewport media capture output */
@@ -462,13 +463,21 @@ public:
 
 	virtual ~SMediaCaptureViewportImage() override
 	{
+		if (TStrongObjectPtr<UMediaProfile> PinnedMediaProfile = MediaProfile.Pin())
+		{
+			PinnedMediaProfile->GetPlaybackManager()->ReleaseManagedViewportFromIndex(MediaOutputIndex, this);
+		}
+		
 		FEditorDelegates::PrePIEEnded.RemoveAll(this);
 		FEditorDelegates::PostPIEStarted.RemoveAll(this);
 	}
 	
-	void Construct(const FArguments& InArgs, const TSharedPtr<FMediaProfileEditor::FViewportCapture> InViewportCapture, const TArray<TLazyObjectPtr<AActor>>& InLockedActors)
+	void Construct(const FArguments& InArgs, UMediaProfile* InMediaProfile, const FMediaFrameworkCaptureCameraViewportCameraOutputInfo& OutputInfo)
 	{
-		ViewportCapture = InViewportCapture;
+		MediaProfile = InMediaProfile;
+		MediaOutputIndex = MediaProfile->FindMediaOutputIndex(OutputInfo.MediaOutput);
+
+		ViewportClient = StaticCastSharedPtr<FLevelEditorViewportClient>(InMediaProfile->GetPlaybackManager()->GetOrCreateManagedViewportFromIndex(MediaOutputIndex, OutputInfo.ViewMode, this));
 
 		FEditorDelegates::PostPIEStarted.AddSP(this, &SMediaCaptureViewportImage::OnPostPIEStarted);
 		FEditorDelegates::PrePIEEnded.AddSP(this, &SMediaCaptureViewportImage::OnPrePIEEnded);
@@ -479,16 +488,19 @@ public:
 		
 		SMediaCaptureImageBase::Construct(BaseArgs);
 
-		Refresh(InLockedActors);
+		Refresh(OutputInfo.LockedActors);
 	}
 
 	virtual void Tick(const FGeometry& AllottedGeometry, const double InCurrentTime, const float InDeltaTime) override
 	{
-		if (PendingViewportSizeChange.IsSet() && ViewportCapture.IsValid())
+		if (PendingViewportSizeChange.IsSet())
 		{
-			const FIntPoint Size = PendingViewportSizeChange.GetValue();
-			ViewportCapture->Viewport->SetFixedViewportSize(Size.X, Size.Y);
-			PendingViewportSizeChange.Reset();
+			if (FSceneViewport* SceneViewport = GetSceneViewport())
+			{
+				const FIntPoint Size = PendingViewportSizeChange.GetValue();
+				SceneViewport->SetFixedViewportSize(Size.X, Size.Y);
+				PendingViewportSizeChange.Reset();
+			}
 		}
 	}
 	
@@ -521,7 +533,7 @@ public:
 	{
 		CurrentLockedActor = InLockedActorIndex;
 
-		if (!ViewportCapture.IsValid() || !ViewportCapture->ViewportClient.IsValid())
+		if (!ViewportClient.IsValid())
 		{
 			return;
 		}
@@ -543,7 +555,7 @@ public:
 			}
 		}
 
-		ViewportCapture->ViewportClient->SetActorLock(ActorToLockTo);
+		ViewportClient->SetActorLock(ActorToLockTo);
 	}
 
 	/** Queues up a desired viewport size that will be set on the next tick */
@@ -552,12 +564,18 @@ public:
 		PendingViewportSizeChange = InDesiredViewportSize;
 	}
 
+	/** Gets the viewport client being displayed by this widget */
+	TSharedPtr<FLevelEditorViewportClient> GetViewportClient() const { return ViewportClient; }
+	
 protected:
 	virtual TSharedRef<SWidget> GetContentWidget() override
 	{
-		if (ViewportCapture.IsValid() && ViewportCapture->ViewportWidget.IsValid())
+		if (FSceneViewport* SceneViewport = GetSceneViewport())
 		{
-			return ViewportCapture->ViewportWidget.ToSharedRef();
+			if (TSharedPtr<SViewport> ViewportWidget = SceneViewport->GetViewportWidget().Pin())
+			{
+				return ViewportWidget.ToSharedRef();
+			}
 		}
 		
 		return SNew(SBorder)
@@ -571,6 +589,22 @@ protected:
 	}
 	
 private:
+	FSceneViewport* GetSceneViewport()
+	{
+		if (!ViewportClient.IsValid())
+		{
+			return nullptr;
+		}
+
+		FViewport* Viewport = ViewportClient->Viewport;
+		if (Viewport && Viewport->GetViewportType() == TEXT("SceneViewport"))
+		{
+			return static_cast<FSceneViewport*>(Viewport);
+		}
+		
+		return nullptr;
+	}
+	
 	/** Raised when Play-in-Editor is started */
 	void OnPostPIEStarted(bool bIsSimulating)
 	{
@@ -584,11 +618,14 @@ private:
 	}
 	
 private:
-	/** The viewport capture that contains the viewport client and viewport widgets to display */
-	TSharedPtr<FMediaProfileEditor::FViewportCapture> ViewportCapture;
+	/** The media profile whose media output is being displayed */
+	TWeakObjectPtr<UMediaProfile> MediaProfile;
 
-	/** Command list that manages the show flags for the viewport capture */
-	TSharedPtr<FUICommandList> ShowFlagsCommandList;
+	/** The index of the media output being displayed */
+	int32 MediaOutputIndex = INDEX_NONE;
+
+	/** The viewport client being used by the media output for capture */
+	TSharedPtr<FLevelEditorViewportClient> ViewportClient;
 
 	/** List of actors that can function as cameras for the viewport capture */
 	TArray<TWeakObjectPtr<AActor>> LockedActors;
@@ -709,6 +746,13 @@ void SMediaProfileMediaOutputDisplay::ConfigureMediaImage()
 		ShowFlagsCommandList->UnmapAction(FlagCommand.ShowMenuItem);
 	}
 	bDisplayShowFlags = false;
+
+	UMediaProfile* MediaProfile = GetMediaProfile();
+	if (!MediaProfile)
+	{
+		MediaImageContainer->SetContent(SNullWidget::NullWidget);
+		return;
+	}
 	
 	UMediaOutput* MediaOutput = GetMediaItem();
 	if (!MediaOutput)
@@ -729,7 +773,7 @@ void SMediaProfileMediaOutputDisplay::ConfigureMediaImage()
 		if (CaptureSettings->CurrentViewportMediaOutput.MediaOutput == MediaOutput)
 		{
 			MediaImageContainer->SetContent(
-				SAssignNew(ImageWidget, SMediaCaptureEditorViewportImage, MediaProfileEditor)
+				SAssignNew(ImageWidget, SMediaCaptureEditorViewportImage, MediaProfile, CaptureSettings->CurrentViewportMediaOutput)
 				.MediaCaptureState(this, &SMediaProfileMediaOutputDisplay::GetMediaOutputCaptureState)
 				.AspectRatio(this, &SMediaProfileMediaOutputDisplay::GetMediaOutputDesiredAspectRatio)
 			);
@@ -744,29 +788,25 @@ void SMediaProfileMediaOutputDisplay::ConfigureMediaImage()
 		
 		if (ViewportOutputInfoIndex != INDEX_NONE)
 		{
-			TSharedPtr<FMediaProfileEditor::FViewportCapture> ViewportCapture = MediaProfileEditor.IsValid() ? MediaProfileEditor.Pin()->FindOrCreateViewportCapture(GetMediaItem()) : nullptr;
-			if (ViewportCapture.IsValid())
+			TSharedPtr<SMediaCaptureViewportImage> ViewportImage;
+			MediaImageContainer->SetContent(
+				SAssignNew(ViewportImage, SMediaCaptureViewportImage, MediaProfile, CaptureSettings->ViewportCaptures[ViewportOutputInfoIndex])
+				.MediaCaptureState(this, &SMediaProfileMediaOutputDisplay::GetMediaOutputCaptureState)
+				.AspectRatio(this, &SMediaProfileMediaOutputDisplay::GetMediaOutputDesiredAspectRatio)
+			);
+
+			FIntPoint TargetSize = MediaOutput->GetRequestedSize();
+			if (TargetSize == UMediaOutput::RequestCaptureSourceSize)
 			{
-				TSharedPtr<SMediaCaptureViewportImage> ViewportImage;
-				MediaImageContainer->SetContent(
-					SAssignNew(ViewportImage, SMediaCaptureViewportImage, ViewportCapture, CaptureSettings->ViewportCaptures[ViewportOutputInfoIndex].LockedActors)
-					.MediaCaptureState(this, &SMediaProfileMediaOutputDisplay::GetMediaOutputCaptureState)
-					.AspectRatio(this, &SMediaProfileMediaOutputDisplay::GetMediaOutputDesiredAspectRatio)
-				);
-
-				FIntPoint TargetSize = MediaOutput->GetRequestedSize();
-				if (TargetSize == UMediaOutput::RequestCaptureSourceSize)
-				{
-					TargetSize = GetMutableDefault<UMediaFrameworkMediaCaptureSettings>()->DefaultViewportCaptureSize;
-				}
-				
-				ViewportImage->QueueViewportSizeChange(TargetSize);
-				ImageWidget = ViewportImage;
-				
-				FShowFlagMenuCommands::Get().BindCommands(*ShowFlagsCommandList.Get(), ViewportCapture->ViewportClient);
-				bDisplayShowFlags = true;
-				return;
+				TargetSize = GetMutableDefault<UMediaFrameworkMediaCaptureSettings>()->DefaultViewportCaptureSize;
 			}
+			
+			ViewportImage->QueueViewportSizeChange(TargetSize);
+			ImageWidget = ViewportImage;
+			
+			FShowFlagMenuCommands::Get().BindCommands(*ShowFlagsCommandList.Get(), ViewportImage->GetViewportClient());
+			bDisplayShowFlags = true;
+			return;
 		}
 
 		RenderTargetOutputInfoIndex = CaptureSettings->RenderTargetCaptures.IndexOfByPredicate(
@@ -921,13 +961,16 @@ void SMediaProfileMediaOutputDisplay::OnObjectPropertyChanged(UObject* InObject,
 
 TOptional<EMediaCaptureState> SMediaProfileMediaOutputDisplay::GetMediaOutputCaptureState() const
 {
-	TSharedPtr<FMediaProfileEditor> PinnedMediaProfileEditor = MediaProfileEditor.Pin();
-	if (!PinnedMediaProfileEditor.IsValid())
+	UMediaProfile* MediaProfile = GetMediaProfile();
+	if (!MediaProfile)
 	{
 		return TOptional<EMediaCaptureState>();
 	}
 
-	return PinnedMediaProfileEditor->GetMediaOutputCaptureState(GetMediaItem());
+	bool bHasError = false;
+	TOptional<EMediaCaptureState> CaptureState = MediaProfile->GetPlaybackManager()->GetOutputCaptureStateFromIndex(MediaItemIndex, bHasError);
+
+	return bHasError ? EMediaCaptureState::Error : CaptureState;
 }
 
 FOptionalSize SMediaProfileMediaOutputDisplay::GetMediaOutputDesiredAspectRatio() const
