Media Profile: Created variant of media capture settings for media profile that can properly be saved to config, and deprecated the use of lazy pointers to store camera actors in the viewport capture settings to avoid a fatal crash when attempting to load capture settings from a config on engine startup

diff --git a/Engine/Plugins/Media/MediaFrameworkUtilities/Source/MediaFrameworkUtilitiesEditor/Private/AssetEditor/MediaProfileEditor.cpp b/Engine/Plugins/Media/MediaFrameworkUtilities/Source/MediaFrameworkUtilitiesEditor/Private/AssetEditor/MediaProfileEditor.cpp
--- a/Engine/Plugins/Media/MediaFrameworkUtilities/Source/MediaFrameworkUtilitiesEditor/Private/AssetEditor/MediaProfileEditor.cpp
+++ b/Engine/Plugins/Media/MediaFrameworkUtilities/Source/MediaFrameworkUtilitiesEditor/Private/AssetEditor/MediaProfileEditor.cpp
@@ -7,6 +7,7 @@
 #include "LevelEditor.h"
 #include "Engine/TextureRenderTarget2D.h"
 #include "LevelEditorViewport.h"
+#include "MediaProfileEditorUserSettings.h"
 #include "PropertyEditorModule.h"
 #include "SMediaProfileDetailsPanel.h"
 #include "SMediaProfileSourcesTreeView.h"
@@ -68,7 +69,6 @@ void FMediaProfileEditor::Initialize(const EToolkitMode::Type Mode, const TShare
 	GEngine->OnLevelActorDeleted().AddSP(this, &FMediaProfileEditor::OnLevelActorsRemoved);
 	FEditorDelegates::OnAssetsDeleted.AddSP(this, &FMediaProfileEditor::OnAssetsDeleted);
 	FCoreUObjectDelegates::OnPreObjectPropertyChanged.AddSP(this, &FMediaProfileEditor::OnObjectPreEditChange);
-	FCoreUObjectDelegates::OnObjectModified.AddSP(this, &FMediaProfileEditor::OnObjectModified);
 	FCoreUObjectDelegates::OnObjectPropertyChanged.AddSP(this, &FMediaProfileEditor::OnObjectPropertyChanged);
 	
 	const TSharedRef<FTabManager::FLayout> StandaloneDefaultLayout = FTabManager::NewLayout("Standalone_MediaProfileEditor_Layout")
@@ -127,7 +127,6 @@ FMediaProfileEditor::~FMediaProfileEditor()
 {
 	GEditor->OnLevelViewportClientListChanged().RemoveAll(this);
 	FCoreUObjectDelegates::OnObjectPropertyChanged.RemoveAll(this);
-	FCoreUObjectDelegates::OnObjectModified.RemoveAll(this);
 	FCoreUObjectDelegates::OnPreObjectPropertyChanged.RemoveAll(this);
 	FEditorDelegates::OnAssetsDeleted.RemoveAll(this);
 	GEngine->OnLevelActorDeleted().RemoveAll(this);
@@ -239,7 +238,7 @@ bool FMediaProfileEditor::CanMediaOutputCapture(UMediaOutput* InMediaOutput) con
 		return false;
 	}
 	
-	UMediaFrameworkWorldSettingsAssetUserData* CaptureSettings = FindOrAddMediaFrameworkCaptureSettings();
+	UMediaProfileEditorCaptureSettings* CaptureSettings = GetMediaFrameworkCaptureSettings();
 	if (!CaptureSettings)
 	{
 		return false;
@@ -261,7 +260,7 @@ bool FMediaProfileEditor::CanMediaOutputCapture(UMediaOutput* InMediaOutput) con
 			continue;
 		}
 
-		const bool bHasValidCamera = ViewportOutputInfo.LockedActors.ContainsByPredicate([](const TLazyObjectPtr<AActor>& ActorRef)
+		const bool bHasValidCamera = ViewportOutputInfo.Cameras.ContainsByPredicate([](const TSoftObjectPtr<AActor>& ActorRef)
 		{
 			return ActorRef.IsValid();
 		});
@@ -290,34 +289,9 @@ bool FMediaProfileEditor::CanMediaOutputCapture(UMediaOutput* InMediaOutput) con
 	return bHasValidCaptureSettings;
 }
 
-UMediaFrameworkWorldSettingsAssetUserData* FMediaProfileEditor::FindOrAddMediaFrameworkCaptureSettings()
+UMediaProfileEditorCaptureSettings* FMediaProfileEditor::GetMediaFrameworkCaptureSettings()
 {
-	UMediaFrameworkWorldSettingsAssetUserData* Result = nullptr;
-
-	if (GetDefault<UMediaFrameworkEditorCaptureSettings>()->bSaveCaptureSetingsInWorld)
-	{
-		UWorld* World = GEditor ? GEditor->GetEditorWorldContext(false).World() : nullptr;
-		if (World)
-		{
-			AWorldSettings* WorldSetting = World ? World->GetWorldSettings() : nullptr;
-			if (WorldSetting)
-			{
-				Result = CastChecked<UMediaFrameworkWorldSettingsAssetUserData>(WorldSetting->GetAssetUserDataOfClass(UMediaFrameworkWorldSettingsAssetUserData::StaticClass()), ECastCheckedType::NullAllowed);
-
-				if (!Result)
-				{
-					Result = NewObject<UMediaFrameworkWorldSettingsAssetUserData>(WorldSetting);
-					WorldSetting->AddAssetUserData(Result);
-				}
-			}
-		}
-	}
-	else
-	{
-		Result = GetMutableDefault<UMediaFrameworkEditorCaptureSettings>();
-	}
-
-	return Result;
+	return GetMutableDefault<UMediaProfileEditorCaptureSettings>();
 }
 
 TSharedRef<SDockTab> FMediaProfileEditor::SpawnTab_MediaOutput(const FSpawnTabArgs& InArgs)
@@ -401,8 +375,6 @@ void FMediaProfileEditor::RefreshSelectedMediaItems(const TArray<int32>& InMedia
 			PlaybackManager->OpenSourceFromIndex(SrcIndex, this);
 		}
 	}
-
-	const bool bAutoRestartCapture = GetDefault<UMediaFrameworkEditorCaptureSettings>()->bAutoRestartCaptureOnChange;
 	
 	for (int32 Index = 0; Index < InMediaOutputs.Num(); ++Index)
 	{
@@ -414,18 +386,18 @@ void FMediaProfileEditor::RefreshSelectedMediaItems(const TArray<int32>& InMedia
 				constexpr bool bDestroyCaptureObjects = false;
 				PlaybackManager->CloseOutput(MediaOutput, bDestroyCaptureObjects);
 
-				if (UMediaFrameworkWorldSettingsAssetUserData* CaptureSettings = FindOrAddMediaFrameworkCaptureSettings())
+				if (UMediaProfileEditorCaptureSettings* CaptureSettings = GetMediaFrameworkCaptureSettings())
 				{
 					if (CaptureSettings->CurrentViewportMediaOutput.MediaOutput == MediaOutput)
 					{
-						PlaybackManager->OpenActiveViewportOutput(MediaOutput, CaptureSettings->CurrentViewportMediaOutput.CaptureOptions, bAutoRestartCapture);
+						PlaybackManager->OpenActiveViewportOutput(MediaOutput, CaptureSettings->CurrentViewportMediaOutput.CaptureOptions, CaptureSettings->bAutoRestartCaptureOnChange);
 					}
 
 					for (const FMediaFrameworkCaptureCameraViewportCameraOutputInfo& ViewportCapture : CaptureSettings->ViewportCaptures)
 					{
 						if (ViewportCapture.MediaOutput == MediaOutput)
 						{
-							PlaybackManager->OpenManagedViewportOutput(MediaOutput, ViewportCapture.CaptureOptions, bAutoRestartCapture);
+							PlaybackManager->OpenManagedViewportOutput(MediaOutput, ViewportCapture.CaptureOptions, CaptureSettings->bAutoRestartCaptureOnChange);
 						}
 					}
 
@@ -433,7 +405,7 @@ void FMediaProfileEditor::RefreshSelectedMediaItems(const TArray<int32>& InMedia
 					{
 						if (RenderTargetCapture.MediaOutput == MediaOutput)
 						{
-							PlaybackManager->OpenRenderTargetOutput(MediaOutput, RenderTargetCapture.RenderTarget, RenderTargetCapture.CaptureOptions, bAutoRestartCapture);
+							PlaybackManager->OpenRenderTargetOutput(MediaOutput, RenderTargetCapture.RenderTarget, RenderTargetCapture.CaptureOptions, CaptureSettings->bAutoRestartCaptureOnChange);
 						}
 					}
 				}
@@ -484,8 +456,6 @@ void FMediaProfileEditor::OnMapChange(uint32 InMapFlags)
 	}
 
 	CloseAllMediaOutputs();
-	
-	OnCaptureSettingsObjectChanged.Broadcast();
 }
 
 void FMediaProfileEditor::OnLevelActorsRemoved(AActor* InActor)
@@ -495,7 +465,7 @@ void FMediaProfileEditor::OnLevelActorsRemoved(AActor* InActor)
 		return;
 	}
 	
-	UMediaFrameworkWorldSettingsAssetUserData* CaptureSettings = FindOrAddMediaFrameworkCaptureSettings();
+	UMediaProfileEditorCaptureSettings* CaptureSettings = GetMediaFrameworkCaptureSettings();
 	if (!CaptureSettings)
 	{
 		return;
@@ -504,7 +474,7 @@ void FMediaProfileEditor::OnLevelActorsRemoved(AActor* InActor)
 	// If the removed actor is referenced by one of the live captures (e.g. a camera in a viewport capture) we need to stop the capture
 	for (const FMediaFrameworkCaptureCameraViewportCameraOutputInfo& ViewportOutputInfo : CaptureSettings->ViewportCaptures)
 	{
-		if (ViewportOutputInfo.LockedActors.Contains(InActor))
+		if (ViewportOutputInfo.Cameras.Contains(InActor))
 		{
 			MediaProfileBeingEdited->GetPlaybackManager()->CloseOutput(ViewportOutputInfo.MediaOutput);
 		}
@@ -518,7 +488,7 @@ void FMediaProfileEditor::OnAssetsDeleted(const TArray<UClass*>& DeletedAssetCla
 		return;
 	}
 	
-	UMediaFrameworkWorldSettingsAssetUserData* CaptureSettings = FindOrAddMediaFrameworkCaptureSettings();
+	UMediaProfileEditorCaptureSettings* CaptureSettings = GetMediaFrameworkCaptureSettings();
 	if (!CaptureSettings)
 	{
 		return;
@@ -547,32 +517,13 @@ void FMediaProfileEditor::OnObjectPreEditChange(UObject* Object, const FEditProp
 	}
 	
 	// If the capture settings are about to be changed, stop all captures
-	UMediaFrameworkWorldSettingsAssetUserData* CaptureSettings = FindOrAddMediaFrameworkCaptureSettings();
+	UMediaProfileEditorCaptureSettings* CaptureSettings = GetMediaFrameworkCaptureSettings();
 	if (Object == CaptureSettings)
 	{
 		CloseAllMediaOutputs();
 	}
 }
 
-void FMediaProfileEditor::OnObjectModified(UObject* Object)
-{
-	const UMediaFrameworkEditorCaptureSettings* GlobalCaptureSettings = GetDefault<UMediaFrameworkEditorCaptureSettings>();
-	if (Object == GlobalCaptureSettings)
-	{
-		bool bCachedSaveCaptureSettingsInWorld = GlobalCaptureSettings->bSaveCaptureSetingsInWorld;
-		
-		// Defer invoking the delegate to ensure any modifications to the global settings complete
-		GEditor->GetTimerManager()->SetTimerForNextTick([this, bCachedSaveCaptureSettingsInWorld, GlobalCaptureSettings]()
-		{
-			// If the capture settings object has changed to or from the global settings, invoke the CaptureSettingsObjectChanged delegate
-			if (bCachedSaveCaptureSettingsInWorld != GlobalCaptureSettings->bSaveCaptureSetingsInWorld)
-			{
-				OnCaptureSettingsObjectChanged.Broadcast();
-			}
-		});
-	}
-}
-
 void FMediaProfileEditor::OnObjectPropertyChanged(UObject* InObject, FPropertyChangedEvent& InPropertyChangedEvent)
 {
 	if (!IsValid(MediaProfileBeingEdited))
@@ -608,7 +559,6 @@ void FMediaProfileEditor::OnObjectPropertyChanged(UObject* InObject, FPropertyCh
 
 		if (UMediaOutput* MediaOutput = Cast<UMediaOutput>(InObject))
 		{
-			const bool bAutoRestartCapture = GetDefault<UMediaFrameworkEditorCaptureSettings>()->bAutoRestartCaptureOnChange;
 			const int32 OutputIndex = MediaProfileBeingEdited->FindMediaOutputIndex(MediaOutput);
 			if (OutputIndex != INDEX_NONE)
 			{
@@ -617,18 +567,18 @@ void FMediaProfileEditor::OnObjectPropertyChanged(UObject* InObject, FPropertyCh
 					constexpr bool bDestroyCaptureObjects = false;
 					PlaybackManager->CloseOutput(MediaOutput, bDestroyCaptureObjects);
 
-					if (UMediaFrameworkWorldSettingsAssetUserData* CaptureSettings = FindOrAddMediaFrameworkCaptureSettings())
+					if (UMediaProfileEditorCaptureSettings* CaptureSettings = GetMediaFrameworkCaptureSettings())
 					{
 						if (CaptureSettings->CurrentViewportMediaOutput.MediaOutput == MediaOutput)
 						{
-							PlaybackManager->OpenActiveViewportOutput(MediaOutput, CaptureSettings->CurrentViewportMediaOutput.CaptureOptions, bAutoRestartCapture);
+							PlaybackManager->OpenActiveViewportOutput(MediaOutput, CaptureSettings->CurrentViewportMediaOutput.CaptureOptions, CaptureSettings->bAutoRestartCaptureOnChange);
 						}
 
 						for (const FMediaFrameworkCaptureCameraViewportCameraOutputInfo& ViewportCapture : CaptureSettings->ViewportCaptures)
 						{
 							if (ViewportCapture.MediaOutput == MediaOutput)
 							{
-								PlaybackManager->OpenManagedViewportOutput(MediaOutput, ViewportCapture.CaptureOptions, bAutoRestartCapture);
+								PlaybackManager->OpenManagedViewportOutput(MediaOutput, ViewportCapture.CaptureOptions, CaptureSettings->bAutoRestartCaptureOnChange);
 							}
 						}
 
@@ -636,7 +586,7 @@ void FMediaProfileEditor::OnObjectPropertyChanged(UObject* InObject, FPropertyCh
 						{
 							if (RenderTargetCapture.MediaOutput == MediaOutput)
 							{
-								PlaybackManager->OpenRenderTargetOutput(MediaOutput, RenderTargetCapture.RenderTarget, RenderTargetCapture.CaptureOptions, bAutoRestartCapture);
+								PlaybackManager->OpenRenderTargetOutput(MediaOutput, RenderTargetCapture.RenderTarget, RenderTargetCapture.CaptureOptions, CaptureSettings->bAutoRestartCaptureOnChange);
 							}
 						}
 					}
diff --git a/Engine/Plugins/Media/MediaFrameworkUtilities/Source/MediaFrameworkUtilitiesEditor/Private/AssetEditor/MediaProfileEditor.h b/Engine/Plugins/Media/MediaFrameworkUtilities/Source/MediaFrameworkUtilitiesEditor/Private/AssetEditor/MediaProfileEditor.h
--- a/Engine/Plugins/Media/MediaFrameworkUtilities/Source/MediaFrameworkUtilitiesEditor/Private/AssetEditor/MediaProfileEditor.h
+++ b/Engine/Plugins/Media/MediaFrameworkUtilities/Source/MediaFrameworkUtilitiesEditor/Private/AssetEditor/MediaProfileEditor.h
@@ -14,6 +14,7 @@ class UMediaFrameworkWorldSettingsAssetUserData;
 class UMediaOutput;
 class UMediaCapture;
 class UMediaProfile;
+class UMediaProfileEditorCaptureSettings;
 class SMediaProfileDetailsPanel;
 
 /**
@@ -64,11 +65,8 @@ public:
 	/** Determines if the specified media output is properly configured for capturing */
 	bool CanMediaOutputCapture(UMediaOutput* InMediaOutput) const;
 
-	/** Gets a valid media capture settings object if one is found, or attempts to add one to the world if one is not found */
-	static UMediaFrameworkWorldSettingsAssetUserData* FindOrAddMediaFrameworkCaptureSettings();
-
-	/** Gets the delegate that is raised when the capture settings object has been changed for any reason */
-	FSimpleMulticastDelegate& GetOnCaptureSettingsObjectChanged() { return OnCaptureSettingsObjectChanged; }
+	/** Gets the media capture settings config object */
+	static UMediaProfileEditorCaptureSettings* GetMediaFrameworkCaptureSettings();
 
 	/** Gets the delegate that is raised when the capture method is changed through the details panel for a media output */
 	FOnCaptureMethodChanged& GetOnCaptureMethodChanged() { return OnCaptureMethodChanged; }
@@ -91,7 +89,6 @@ private:
 	void OnLevelActorsRemoved(AActor* InActor);
 	void OnAssetsDeleted(const TArray<UClass*>& DeletedAssetClasses);
 	void OnObjectPreEditChange(UObject* Object, const FEditPropertyChain& PropertyChain);
-	void OnObjectModified(UObject* Object);
 	void OnObjectPropertyChanged(UObject* InObject, FPropertyChangedEvent& InPropertyChangedEvent);
 	
 private:
@@ -107,9 +104,6 @@ private:
 	/** List of media outputs whose capture needs to be restarted on the next engine tick */
 	TSet<TWeakObjectPtr<UMediaOutput>> OutputsToRestartCapture;
 
-	/** Raised when the object that stores the capture settings has changed */
-	FSimpleMulticastDelegate OnCaptureSettingsObjectChanged;
-
 	/** Raised when the capture method is changed on a media output through the details panel */
 	FOnCaptureMethodChanged OnCaptureMethodChanged;
 	
diff --git a/Engine/Plugins/Media/MediaFrameworkUtilities/Source/MediaFrameworkUtilitiesEditor/Private/AssetEditor/MediaProfileEditorUserSettings.h b/Engine/Plugins/Media/MediaFrameworkUtilities/Source/MediaFrameworkUtilitiesEditor/Private/AssetEditor/MediaProfileEditorUserSettings.h
--- a/Engine/Plugins/Media/MediaFrameworkUtilities/Source/MediaFrameworkUtilitiesEditor/Private/AssetEditor/MediaProfileEditorUserSettings.h
+++ b/Engine/Plugins/Media/MediaFrameworkUtilities/Source/MediaFrameworkUtilitiesEditor/Private/AssetEditor/MediaProfileEditorUserSettings.h
@@ -3,6 +3,7 @@
 #pragma once
 
 #include "UObject/Object.h"
+#include "MediaFrameworkWorldSettingsAssetUserData.h"
 
 #include "MediaProfileEditorUserSettings.generated.h"
 
@@ -20,4 +21,18 @@ public:
 	/** Whether to show the genlock entry in the viewport toolbar */
 	UPROPERTY(Config)
 	bool bShowGenlockInViewportToolbar = true;
+};
+
+/**
+ * Variant of the media capture settings config for media outputs in the media profile
+ */
+UCLASS(MinimalAPI, config = EditorSettings)
+class UMediaProfileEditorCaptureSettings : public UMediaFrameworkWorldSettingsAssetUserData
+{
+	GENERATED_BODY()
+	
+public:
+	/** Should the capture be restarted if the media output is modified. */
+	UPROPERTY(config)
+	bool bAutoRestartCaptureOnChange = true;
 };
\ No newline at end of file
diff --git a/Engine/Plugins/Media/MediaFrameworkUtilities/Source/MediaFrameworkUtilitiesEditor/Private/AssetEditor/MediaProfileMediaItemDetailCustomization.cpp b/Engine/Plugins/Media/MediaFrameworkUtilities/Source/MediaFrameworkUtilitiesEditor/Private/AssetEditor/MediaProfileMediaItemDetailCustomization.cpp
--- a/Engine/Plugins/Media/MediaFrameworkUtilities/Source/MediaFrameworkUtilitiesEditor/Private/AssetEditor/MediaProfileMediaItemDetailCustomization.cpp
+++ b/Engine/Plugins/Media/MediaFrameworkUtilities/Source/MediaFrameworkUtilitiesEditor/Private/AssetEditor/MediaProfileMediaItemDetailCustomization.cpp
@@ -6,6 +6,7 @@
 #include "IPropertyUtilities.h"
 #include "IStructureDataProvider.h"
 #include "MediaFrameworkWorldSettingsAssetUserData.h"
+#include "MediaProfileEditorUserSettings.h"
 #include "Algo/Count.h"
 #include "Profile/MediaProfilePlaybackManager.h"
 #include "UI/MediaFrameworkUtilitiesEditorStyle.h"
@@ -126,7 +127,7 @@ public:
 private:
 	void OnCurrentViewportMethodSelected()
 	{
-		UMediaFrameworkWorldSettingsAssetUserData* MediaCaptureSettings = FMediaProfileEditor::FindOrAddMediaFrameworkCaptureSettings();
+		UMediaProfileEditorCaptureSettings* MediaCaptureSettings = FMediaProfileEditor::GetMediaFrameworkCaptureSettings();
 		if (!MediaCaptureSettings)
 		{
 			return;
@@ -166,12 +167,13 @@ private:
 			MediaCaptureSettings->CurrentViewportMediaOutput.MediaOutput = MediaOutput;
 		}
 
+		MediaCaptureSettings->SaveConfig();
 		OnMethodChanged.ExecuteIfBound();
 	}
 	
 	void OnViewportMethodSelected()
 	{
-		UMediaFrameworkWorldSettingsAssetUserData* MediaCaptureSettings = FMediaProfileEditor::FindOrAddMediaFrameworkCaptureSettings();
+		UMediaProfileEditorCaptureSettings* MediaCaptureSettings = FMediaProfileEditor::GetMediaFrameworkCaptureSettings();
 		if (!MediaCaptureSettings)
 		{
 			return;
@@ -216,12 +218,13 @@ private:
 			MediaCaptureSettings->RenderTargetCaptures.RemoveAll(GetFindOutputInfoPredicate<FMediaFrameworkCaptureRenderTargetCameraOutputInfo>(MediaOutput));
 		}
 
+		MediaCaptureSettings->SaveConfig();
 		OnMethodChanged.ExecuteIfBound();
 	}
 
 	void OnRenderTargetMethodSelected()
 	{
-		UMediaFrameworkWorldSettingsAssetUserData* MediaCaptureSettings = FMediaProfileEditor::FindOrAddMediaFrameworkCaptureSettings();
+		UMediaProfileEditorCaptureSettings* MediaCaptureSettings = FMediaProfileEditor::GetMediaFrameworkCaptureSettings();
 		if (!MediaCaptureSettings)
 		{
 			return;
@@ -266,12 +269,13 @@ private:
 			MediaCaptureSettings->ViewportCaptures.RemoveAll(GetFindOutputInfoPredicate<FMediaFrameworkCaptureCameraViewportCameraOutputInfo>(MediaOutput));
 		}
 
+		MediaCaptureSettings->SaveConfig();
 		OnMethodChanged.ExecuteIfBound();
 	}
 
 	FText GetButtonText() const
 	{
-		UMediaFrameworkWorldSettingsAssetUserData* MediaCaptureSettings = FMediaProfileEditor::FindOrAddMediaFrameworkCaptureSettings();
+		UMediaProfileEditorCaptureSettings* MediaCaptureSettings = FMediaProfileEditor::GetMediaFrameworkCaptureSettings();
 		if (!MediaCaptureSettings)
 		{
 			return LOCTEXT("MediaCaptureMethodNotConfigured", "Not Configured");
@@ -351,8 +355,9 @@ protected:
 	
 	void PropertyValueChange(const FPropertyChangedEvent& PropertyChangedEvent)
 	{
-		if (UMediaFrameworkWorldSettingsAssetUserData* MediaCaptureSettings = FMediaProfileEditor::FindOrAddMediaFrameworkCaptureSettings())
+		if (UMediaProfileEditorCaptureSettings* MediaCaptureSettings = FMediaProfileEditor::GetMediaFrameworkCaptureSettings())
 		{
+			MediaCaptureSettings->SaveConfig();
 			FCoreUObjectDelegates::OnObjectPropertyChanged.Broadcast(MediaCaptureSettings, const_cast<FPropertyChangedEvent&>(PropertyChangedEvent));
 		}
 	}
@@ -365,7 +370,7 @@ protected:
 
 	virtual void CustomizeChildren(TSharedRef<IPropertyHandle> PropertyHandle, IDetailChildrenBuilder& ChildBuilder, IPropertyTypeCustomizationUtils& CustomizationUtils) override
 	{
-		TSharedPtr<IPropertyHandle> LockedActorsHandle = PropertyHandle->GetChildHandle(GET_MEMBER_NAME_CHECKED(FMediaFrameworkCaptureCameraViewportCameraOutputInfo, LockedActors));
+		TSharedPtr<IPropertyHandle> LockedActorsHandle = PropertyHandle->GetChildHandle(GET_MEMBER_NAME_CHECKED(FMediaFrameworkCaptureCameraViewportCameraOutputInfo, Cameras));
 		LockedActorsHandle->SetOnPropertyValueChangedWithData(TDelegate<void(const FPropertyChangedEvent&)>::CreateSP(this, &FMediaFrameworkViewportOutputInfoCustomization::PropertyValueChange));
 		LockedActorsHandle->SetOnChildPropertyValueChangedWithData(TDelegate<void(const FPropertyChangedEvent&)>::CreateSP(this, &FMediaFrameworkViewportOutputInfoCustomization::PropertyValueChange));
 		ChildBuilder.AddProperty(LockedActorsHandle.ToSharedRef());
@@ -384,8 +389,9 @@ protected:
 	
 	void PropertyValueChange(const FPropertyChangedEvent& PropertyChangedEvent)
 	{
-		if (UMediaFrameworkWorldSettingsAssetUserData* MediaCaptureSettings = FMediaProfileEditor::FindOrAddMediaFrameworkCaptureSettings())
+		if (UMediaProfileEditorCaptureSettings* MediaCaptureSettings = FMediaProfileEditor::GetMediaFrameworkCaptureSettings())
 		{
+			MediaCaptureSettings->SaveConfig();
 			FCoreUObjectDelegates::OnObjectPropertyChanged.Broadcast(MediaCaptureSettings, const_cast<FPropertyChangedEvent&>(PropertyChangedEvent));
 		}
 	}
@@ -416,8 +422,9 @@ protected:
 
 	void PropertyValueChange(const FPropertyChangedEvent& InPropertyChangedEvent)
 	{
-		if (UMediaFrameworkWorldSettingsAssetUserData* MediaCaptureSettings = FMediaProfileEditor::FindOrAddMediaFrameworkCaptureSettings())
+		if (UMediaProfileEditorCaptureSettings* MediaCaptureSettings = FMediaProfileEditor::GetMediaFrameworkCaptureSettings())
 		{
+			MediaCaptureSettings->SaveConfig();
 			FCoreUObjectDelegates::OnObjectPropertyChanged.Broadcast(MediaCaptureSettings, const_cast<FPropertyChangedEvent&>(InPropertyChangedEvent));
 		}
 	}
@@ -426,10 +433,6 @@ protected:
 FMediaProfileMediaOutputDetailCustomization::~FMediaProfileMediaOutputDetailCustomization()
 {
 	FCoreUObjectDelegates::OnObjectPropertyChanged.RemoveAll(this);
-	if (MediaProfileEditor.IsValid())
-	{
-		MediaProfileEditor.Pin()->GetOnCaptureSettingsObjectChanged().RemoveAll(this);
-	}
 }
 
 void FMediaProfileMediaOutputDetailCustomization::CustomizeDetails(IDetailLayoutBuilder& DetailBuilder)
@@ -441,7 +444,6 @@ void FMediaProfileMediaOutputDetailCustomization::CustomizeDetails(IDetailLayout
 	// adjust the sort order, but we do need the details builder to lump all categories together, hence the empty sort function
 	DetailBuilder.SortCategories([](const TMap<FName, IDetailCategoryBuilder*>& CategoryMap) { });
 	
-	MediaProfileEditor.Pin()->GetOnCaptureSettingsObjectChanged().AddSP(this, &FMediaProfileMediaOutputDetailCustomization::OnCaptureSettingsObjectChanged);
 	FCoreUObjectDelegates::OnObjectPropertyChanged.AddSP(this, &FMediaProfileMediaOutputDetailCustomization::OnObjectPropertyChanged);
 	
 	TArray<TWeakObjectPtr<UObject>> Objects;
@@ -503,7 +505,7 @@ void FMediaProfileMediaOutputDetailCustomization::CustomizeDetails(IDetailLayout
 	bool bHasCurrentViewportCapture = false;
 	bool bHasViewportCaptures = false;
 	bool bHasRenderTargetCaptures = false;
-	if (UMediaFrameworkWorldSettingsAssetUserData* MediaCaptureSettings = FMediaProfileEditor::FindOrAddMediaFrameworkCaptureSettings())
+	if (UMediaProfileEditorCaptureSettings* MediaCaptureSettings = FMediaProfileEditor::GetMediaFrameworkCaptureSettings())
 	{
 		for (const TWeakObjectPtr<UObject>& WeakObject : Objects)
 		{
@@ -642,7 +644,7 @@ void FMediaProfileMediaOutputDetailCustomization::SetMediaObject(UMediaOutput* I
 
 	PinnedMediaProfile->GetPlaybackManager()->CloseOutputFromIndex(Index);
 	
-	UMediaFrameworkWorldSettingsAssetUserData* CaptureSettings = FMediaProfileEditor::FindOrAddMediaFrameworkCaptureSettings();
+	UMediaProfileEditorCaptureSettings* CaptureSettings = FMediaProfileEditor::GetMediaFrameworkCaptureSettings();
 	if (!CaptureSettings)
 	{
 		return;
@@ -682,6 +684,7 @@ void FMediaProfileMediaOutputDetailCustomization::SetMediaObject(UMediaOutput* I
 		CaptureSettings->CurrentViewportMediaOutput.MediaOutput = InNewMediaObject;
 	}
 	
+	CaptureSettings->SaveConfig();
 	PinnedMediaProfile->SetMediaOutput(Index, InNewMediaObject);
 }
 
@@ -700,14 +703,6 @@ void FMediaProfileMediaOutputDetailCustomization::RegisterCaptureSection(FProper
 	}
 }
 
-void FMediaProfileMediaOutputDetailCustomization::OnCaptureSettingsObjectChanged()
-{
-	if (CachedDetailBuilder.IsValid())
-	{
-		CachedDetailBuilder.Pin()->ForceRefreshDetails();
-	}
-}
-
 void FMediaProfileMediaOutputDetailCustomization::OnCaptureMethodChanged()
 {
 	TSharedPtr<FMediaProfileEditor> PinnedMediaProfileEditor = MediaProfileEditor.Pin();
@@ -753,9 +748,9 @@ void FMediaProfileMediaOutputDetailCustomization::OnCaptureMethodChanged()
 
 void FMediaProfileMediaOutputDetailCustomization::OnObjectPropertyChanged(UObject* InObject, FPropertyChangedEvent& InPropertyChangedEvent)
 {
-	if (UMediaFrameworkWorldSettingsAssetUserData* CaptureSettings = Cast<UMediaFrameworkWorldSettingsAssetUserData>(InObject))
+	if (UMediaProfileEditorCaptureSettings* CaptureSettings = Cast<UMediaProfileEditorCaptureSettings>(InObject))
 	{
-		if (CaptureSettings != FMediaProfileEditor::FindOrAddMediaFrameworkCaptureSettings())
+		if (CaptureSettings != FMediaProfileEditor::GetMediaFrameworkCaptureSettings())
 		{
 			return;
 		}
diff --git a/Engine/Plugins/Media/MediaFrameworkUtilities/Source/MediaFrameworkUtilitiesEditor/Private/AssetEditor/MediaProfileMediaItemDetailCustomization.h b/Engine/Plugins/Media/MediaFrameworkUtilities/Source/MediaFrameworkUtilitiesEditor/Private/AssetEditor/MediaProfileMediaItemDetailCustomization.h
--- a/Engine/Plugins/Media/MediaFrameworkUtilities/Source/MediaFrameworkUtilitiesEditor/Private/AssetEditor/MediaProfileMediaItemDetailCustomization.h
+++ b/Engine/Plugins/Media/MediaFrameworkUtilities/Source/MediaFrameworkUtilitiesEditor/Private/AssetEditor/MediaProfileMediaItemDetailCustomization.h
@@ -367,8 +367,6 @@ private:
 
 	/** Registers a section for the Capture category */
 	void RegisterCaptureSection(FPropertyEditorModule& PropertyModule, const FName& CaptureCategory);
-	
-	void OnCaptureSettingsObjectChanged();
 
 	void OnCaptureMethodChanged();
 	void OnObjectPropertyChanged(UObject* InObject, FPropertyChangedEvent& InPropertyChangedEvent);
diff --git a/Engine/Plugins/Media/MediaFrameworkUtilities/Source/MediaFrameworkUtilitiesEditor/Private/AssetEditor/SMediaProfileDetailsInfoPanel.cpp b/Engine/Plugins/Media/MediaFrameworkUtilities/Source/MediaFrameworkUtilitiesEditor/Private/AssetEditor/SMediaProfileDetailsInfoPanel.cpp
--- a/Engine/Plugins/Media/MediaFrameworkUtilities/Source/MediaFrameworkUtilitiesEditor/Private/AssetEditor/SMediaProfileDetailsInfoPanel.cpp
+++ b/Engine/Plugins/Media/MediaFrameworkUtilities/Source/MediaFrameworkUtilitiesEditor/Private/AssetEditor/SMediaProfileDetailsInfoPanel.cpp
@@ -7,6 +7,7 @@
 #include "MediaFrameworkWorldSettingsAssetUserData.h"
 #include "MediaOutput.h"
 #include "MediaProfileEditor.h"
+#include "MediaProfileEditorUserSettings.h"
 #include "MediaSource.h"
 #include "Engine/TextureRenderTarget2D.h"
 #include "MediaAssets/ProxyMediaOutput.h"
@@ -318,7 +319,7 @@ TSharedRef<SWidget> SMediaObjectInfoPanel::CreateInfoGroupWidget(const TAttribut
 FText SMediaObjectInfoPanel::GetCaptureMethodText() const
 {
 	UMediaOutput* MediaOutput = ObjectAsMediaOutput();
-	UMediaFrameworkWorldSettingsAssetUserData* CaptureSettings = FMediaProfileEditor::FindOrAddMediaFrameworkCaptureSettings();
+	UMediaProfileEditorCaptureSettings* CaptureSettings = FMediaProfileEditor::GetMediaFrameworkCaptureSettings();
 	if (!MediaOutput || !CaptureSettings)
 	{
 		return FText::GetEmpty();
@@ -362,7 +363,7 @@ FText SMediaObjectInfoPanel::GetCaptureMethodText() const
 FText SMediaObjectInfoPanel::GetCaptureObjectLabelText() const
 {
 	UMediaOutput* MediaOutput = ObjectAsMediaOutput();
-	UMediaFrameworkWorldSettingsAssetUserData* CaptureSettings = FMediaProfileEditor::FindOrAddMediaFrameworkCaptureSettings();
+	UMediaProfileEditorCaptureSettings* CaptureSettings = FMediaProfileEditor::GetMediaFrameworkCaptureSettings();
 	if (!MediaOutput || !CaptureSettings)
 	{
 		return FText::GetEmpty();
@@ -406,7 +407,7 @@ FText SMediaObjectInfoPanel::GetCaptureObjectLabelText() const
 FText SMediaObjectInfoPanel::GetCaptureObjectValueText() const
 {
 	UMediaOutput* MediaOutput = ObjectAsMediaOutput();
-	UMediaFrameworkWorldSettingsAssetUserData* CaptureSettings = FMediaProfileEditor::FindOrAddMediaFrameworkCaptureSettings();
+	UMediaProfileEditorCaptureSettings* CaptureSettings = FMediaProfileEditor::GetMediaFrameworkCaptureSettings();
 	if (!MediaOutput || !CaptureSettings)
 	{
 		return FText::GetEmpty();
@@ -425,7 +426,7 @@ FText SMediaObjectInfoPanel::GetCaptureObjectValueText() const
 
 		bHasViewportCaptures = true;
 		
-		for (const TLazyObjectPtr<AActor>& Camera : Info.LockedActors)
+		for (const TSoftObjectPtr<AActor>& Camera : Info.Cameras)
 		{
 			AActor* Actor = Camera.Get();
 			if (Actor)
@@ -510,7 +511,7 @@ FSlateColor SMediaObjectInfoPanel::GetCaptureStatusColor() const
 FText SMediaObjectInfoPanel::GetCaptureColorConversionText() const
 {
 	UMediaOutput* MediaOutput = ObjectAsMediaOutput();
-	UMediaFrameworkWorldSettingsAssetUserData* CaptureSettings = FMediaProfileEditor::FindOrAddMediaFrameworkCaptureSettings();
+	UMediaProfileEditorCaptureSettings* CaptureSettings = FMediaProfileEditor::GetMediaFrameworkCaptureSettings();
 	if (!MediaOutput || !CaptureSettings)
 	{
 		return LOCTEXT("NoneText", "None");
diff --git a/Engine/Plugins/Media/MediaFrameworkUtilities/Source/MediaFrameworkUtilitiesEditor/Private/AssetEditor/SMediaProfileDetailsPanel.cpp b/Engine/Plugins/Media/MediaFrameworkUtilities/Source/MediaFrameworkUtilitiesEditor/Private/AssetEditor/SMediaProfileDetailsPanel.cpp
--- a/Engine/Plugins/Media/MediaFrameworkUtilities/Source/MediaFrameworkUtilitiesEditor/Private/AssetEditor/SMediaProfileDetailsPanel.cpp
+++ b/Engine/Plugins/Media/MediaFrameworkUtilities/Source/MediaFrameworkUtilitiesEditor/Private/AssetEditor/SMediaProfileDetailsPanel.cpp
@@ -10,6 +10,7 @@
 #include "IDetailsView.h"
 #include "MediaSource.h"
 #include "MediaOutput.h"
+#include "MediaProfileEditorUserSettings.h"
 #include "MediaProfileMediaItemDetailCustomization.h"
 #include "SMediaProfileDetailsInfoPanel.h"
 #include "SPositiveActionButton.h"
@@ -265,7 +266,7 @@ void SMediaProfileDetailsPanel::OnCaptureButtonPressed()
 	}
 	else
 	{
-		UMediaFrameworkWorldSettingsAssetUserData* MediaCaptureSettings = FMediaProfileEditor::FindOrAddMediaFrameworkCaptureSettings();
+		UMediaProfileEditorCaptureSettings* MediaCaptureSettings = FMediaProfileEditor::GetMediaFrameworkCaptureSettings();
 		if (!IsValid(MediaCaptureSettings))
 		{
 			return;
@@ -277,7 +278,7 @@ void SMediaProfileDetailsPanel::OnCaptureButtonPressed()
 			{
 				if (MediaCaptureSettings->CurrentViewportMediaOutput.MediaOutput == MediaOutput)
 				{
-					PinnedMediaProfile->GetPlaybackManager()->OpenActiveViewportOutput(MediaOutput, MediaCaptureSettings->CurrentViewportMediaOutput.CaptureOptions);
+					PinnedMediaProfile->GetPlaybackManager()->OpenActiveViewportOutput(MediaOutput, MediaCaptureSettings->CurrentViewportMediaOutput.CaptureOptions, MediaCaptureSettings->bAutoRestartCaptureOnChange);
 					continue;
 				}
 
@@ -293,7 +294,7 @@ void SMediaProfileDetailsPanel::OnCaptureButtonPressed()
 					TSharedPtr<FViewportClient> ViewportClient = PinnedMediaProfile->GetPlaybackManager()->GetOrCreateManagedViewport(MediaOutput, ViewportCapture->ViewMode);
 					if (ViewportClient.IsValid())
 					{
-						PinnedMediaProfile->GetPlaybackManager()->OpenManagedViewportOutput(MediaOutput, ViewportCapture->CaptureOptions);
+						PinnedMediaProfile->GetPlaybackManager()->OpenManagedViewportOutput(MediaOutput, ViewportCapture->CaptureOptions, MediaCaptureSettings->bAutoRestartCaptureOnChange);
 					}
 					continue;
 				}
@@ -306,7 +307,7 @@ void SMediaProfileDetailsPanel::OnCaptureButtonPressed()
 
 				if (RenderTargetCapture)
 				{
-					PinnedMediaProfile->GetPlaybackManager()->OpenRenderTargetOutput(MediaOutput, RenderTargetCapture->RenderTarget, RenderTargetCapture->CaptureOptions);
+					PinnedMediaProfile->GetPlaybackManager()->OpenRenderTargetOutput(MediaOutput, RenderTargetCapture->RenderTarget, RenderTargetCapture->CaptureOptions, MediaCaptureSettings->bAutoRestartCaptureOnChange);
 				}
 			}
 		}
diff --git a/Engine/Plugins/Media/MediaFrameworkUtilities/Source/MediaFrameworkUtilitiesEditor/Private/AssetEditor/SMediaProfileMediaItemDisplay.cpp b/Engine/Plugins/Media/MediaFrameworkUtilities/Source/MediaFrameworkUtilitiesEditor/Private/AssetEditor/SMediaProfileMediaItemDisplay.cpp
--- a/Engine/Plugins/Media/MediaFrameworkUtilities/Source/MediaFrameworkUtilitiesEditor/Private/AssetEditor/SMediaProfileMediaItemDisplay.cpp
+++ b/Engine/Plugins/Media/MediaFrameworkUtilities/Source/MediaFrameworkUtilitiesEditor/Private/AssetEditor/SMediaProfileMediaItemDisplay.cpp
@@ -488,7 +488,7 @@ public:
 		
 		SMediaCaptureImageBase::Construct(BaseArgs);
 
-		Refresh(OutputInfo.LockedActors);
+		Refresh(OutputInfo.Cameras);
 	}
 
 	virtual void Tick(const FGeometry& AllottedGeometry, const double InCurrentTime, const float InDeltaTime) override
@@ -505,11 +505,11 @@ public:
 	}
 	
 	/** Updates the display if the capture's locked actors have been changed */
-	void Refresh(const TArray<TLazyObjectPtr<AActor>>& InLockedActors)
+	void Refresh(const TArray<TSoftObjectPtr<AActor>>& InLockedActors)
 	{
 		LockedActors.Empty();
 		LockedActors.Reserve(InLockedActors.Num());
-		for (const TLazyObjectPtr<AActor>& ActorRef : InLockedActors)
+		for (const TSoftObjectPtr<AActor>& ActorRef : InLockedActors)
 		{
 			AActor* Actor = ActorRef.Get();
 			if (Actor)
@@ -768,7 +768,7 @@ void SMediaProfileMediaOutputDisplay::ConfigureMediaImage()
 		return;
 	}
 	
-	if (UMediaFrameworkWorldSettingsAssetUserData* CaptureSettings = FMediaProfileEditor::FindOrAddMediaFrameworkCaptureSettings())
+	if (UMediaProfileEditorCaptureSettings* CaptureSettings = FMediaProfileEditor::GetMediaFrameworkCaptureSettings())
 	{
 		if (CaptureSettings->CurrentViewportMediaOutput.MediaOutput == MediaOutput)
 		{
@@ -883,9 +883,9 @@ void SMediaProfileMediaOutputDisplay::OnObjectPropertyChanged(UObject* InObject,
 {
 	SMediaProfileMediaItemDisplayBase::OnObjectPropertyChanged(InObject, InPropertyChangedEvent);
 	
-	if (UMediaFrameworkWorldSettingsAssetUserData* CaptureSettings = Cast<UMediaFrameworkWorldSettingsAssetUserData>(InObject))
+	if (UMediaProfileEditorCaptureSettings* CaptureSettings = Cast<UMediaProfileEditorCaptureSettings>(InObject))
 	{
-		if (CaptureSettings != FMediaProfileEditor::FindOrAddMediaFrameworkCaptureSettings())
+		if (CaptureSettings != FMediaProfileEditor::GetMediaFrameworkCaptureSettings())
 		{
 			return;
 		}
@@ -908,7 +908,7 @@ void SMediaProfileMediaOutputDisplay::OnObjectPropertyChanged(UObject* InObject,
 			return;
 		}
 
-		if (PropertyName == GET_MEMBER_NAME_CHECKED(FMediaFrameworkCaptureCameraViewportCameraOutputInfo, LockedActors))
+		if (PropertyName == GET_MEMBER_NAME_CHECKED(FMediaFrameworkCaptureCameraViewportCameraOutputInfo, Cameras))
 		{
 			if (ViewportOutputInfoIndex == INDEX_NONE || !CaptureSettings->ViewportCaptures.IsValidIndex(ViewportOutputInfoIndex))
 			{
@@ -921,7 +921,7 @@ void SMediaProfileMediaOutputDisplay::OnObjectPropertyChanged(UObject* InObject,
 				return;
 			}
 
-			ViewportImageWidget->Refresh(CaptureSettings->ViewportCaptures[ViewportOutputInfoIndex].LockedActors);
+			ViewportImageWidget->Refresh(CaptureSettings->ViewportCaptures[ViewportOutputInfoIndex].Cameras);
 			return;
 		}
 
diff --git a/Engine/Plugins/Media/MediaFrameworkUtilities/Source/MediaFrameworkUtilitiesEditor/Private/AssetEditor/SMediaProfileSourcesTreeView.cpp b/Engine/Plugins/Media/MediaFrameworkUtilities/Source/MediaFrameworkUtilitiesEditor/Private/AssetEditor/SMediaProfileSourcesTreeView.cpp
--- a/Engine/Plugins/Media/MediaFrameworkUtilities/Source/MediaFrameworkUtilitiesEditor/Private/AssetEditor/SMediaProfileSourcesTreeView.cpp
+++ b/Engine/Plugins/Media/MediaFrameworkUtilities/Source/MediaFrameworkUtilitiesEditor/Private/AssetEditor/SMediaProfileSourcesTreeView.cpp
@@ -5,6 +5,7 @@
 #include "Editor.h"
 #include "MediaOutput.h"
 #include "MediaProfileEditor.h"
+#include "MediaProfileEditorUserSettings.h"
 #include "MediaSource.h"
 #include "ScopedTransaction.h"
 #include "SlateOptMacros.h"
@@ -947,7 +948,7 @@ void SMediaProfileSourcesTreeView::DeleteSelectedMedia()
 		MediaProfile->RemoveMediaSource(SourcesToRemove[Index]);
 	}
 
-	UMediaFrameworkWorldSettingsAssetUserData* MediaCaptureSettings = FMediaProfileEditor::FindOrAddMediaFrameworkCaptureSettings();
+	UMediaProfileEditorCaptureSettings* MediaCaptureSettings = FMediaProfileEditor::GetMediaFrameworkCaptureSettings();
 
 	auto RemoveMediaCaptures = [MediaCaptureSettings](UMediaOutput* MediaOutput)
 	{
@@ -962,6 +963,8 @@ void SMediaProfileSourcesTreeView::DeleteSelectedMedia()
 			{
 				return RenderTargetOutputInfo.MediaOutput == MediaOutput;
 			});
+
+			MediaCaptureSettings->SaveConfig();
 		}
 	};
 	
diff --git a/Engine/Plugins/Media/MediaFrameworkUtilities/Source/MediaFrameworkUtilitiesEditor/Private/CaptureTab/MediaFrameworkCapturePanelBlueprintLibrary.cpp b/Engine/Plugins/Media/MediaFrameworkUtilities/Source/MediaFrameworkUtilitiesEditor/Private/CaptureTab/MediaFrameworkCapturePanelBlueprintLibrary.cpp
--- a/Engine/Plugins/Media/MediaFrameworkUtilities/Source/MediaFrameworkUtilitiesEditor/Private/CaptureTab/MediaFrameworkCapturePanelBlueprintLibrary.cpp
+++ b/Engine/Plugins/Media/MediaFrameworkUtilities/Source/MediaFrameworkUtilitiesEditor/Private/CaptureTab/MediaFrameworkCapturePanelBlueprintLibrary.cpp
@@ -90,7 +90,7 @@ void UMediaFrameworkCapturePanel::AddViewportCapture(UMediaOutput* MediaOutput,
 		{
 			FMediaFrameworkCaptureCameraViewportCameraOutputInfo Data;
 			Data.MediaOutput = MediaOutput;
-			Data.LockedActors.Add(Camera);
+			Data.Cameras.Add(Camera);
 			Data.CaptureOptions = CaptureOptions;
 			Data.ViewMode = ViewMode;
 			UserData->ViewportCaptures.Emplace(MoveTemp(Data));
diff --git a/Engine/Plugins/Media/MediaFrameworkUtilities/Source/MediaFrameworkUtilitiesEditor/Private/CaptureTab/SMediaFrameworkCapture.cpp b/Engine/Plugins/Media/MediaFrameworkUtilities/Source/MediaFrameworkUtilitiesEditor/Private/CaptureTab/SMediaFrameworkCapture.cpp
--- a/Engine/Plugins/Media/MediaFrameworkUtilities/Source/MediaFrameworkUtilitiesEditor/Private/CaptureTab/SMediaFrameworkCapture.cpp
+++ b/Engine/Plugins/Media/MediaFrameworkUtilities/Source/MediaFrameworkUtilitiesEditor/Private/CaptureTab/SMediaFrameworkCapture.cpp
@@ -479,8 +479,8 @@ bool SMediaFrameworkCapture::CanEnableViewport() const
 	{
 		for (const FMediaFrameworkCaptureCameraViewportCameraOutputInfo& Info : AssetUserData->ViewportCaptures)
 		{
-			bEnabled = Info.MediaOutput && Info.LockedActors.Num() > 0;
-			for (const TLazyObjectPtr<AActor>& CameraActorRef : Info.LockedActors)
+			bEnabled = Info.MediaOutput && Info.Cameras.Num() > 0;
+			for (const TSoftObjectPtr<AActor>& CameraActorRef : Info.Cameras)
 			{
 				bEnabled = bEnabled && CameraActorRef.IsValid();
 			}
@@ -535,8 +535,8 @@ void SMediaFrameworkCapture::EnabledCapture(bool bEnabled)
 		for (const FMediaFrameworkCaptureCameraViewportCameraOutputInfo& Info : AssetUserData->ViewportCaptures)
 		{
 			TArray<TWeakObjectPtr<AActor>> InfoPreviewActors;
-			InfoPreviewActors.Reserve(Info.LockedActors.Num());
-			for (const TLazyObjectPtr<AActor>& ActorRef : Info.LockedActors)
+			InfoPreviewActors.Reserve(Info.Cameras.Num());
+			for (const TSoftObjectPtr<AActor>& ActorRef : Info.Cameras)
 			{
 				AActor* Actor = ActorRef.Get();
 				if (Actor)
@@ -696,7 +696,7 @@ void SMediaFrameworkCapture::OnLevelActorsRemoved(AActor* InActor)
 	{
 		for (const FMediaFrameworkCaptureCameraViewportCameraOutputInfo& Info : AssetUserData->ViewportCaptures)
 		{
-			if (Info.LockedActors.Contains(InActor))
+			if (Info.Cameras.Contains(InActor))
 			{
 				EnabledCapture(false);
 				return;
diff --git a/Engine/Plugins/Media/MediaFrameworkUtilities/Source/MediaFrameworkUtilitiesEditor/Private/MediaFrameworkWorldSettingsAssetUserData.cpp b/Engine/Plugins/Media/MediaFrameworkUtilities/Source/MediaFrameworkUtilitiesEditor/Private/MediaFrameworkWorldSettingsAssetUserData.cpp
--- a/Engine/Plugins/Media/MediaFrameworkUtilities/Source/MediaFrameworkUtilitiesEditor/Private/MediaFrameworkWorldSettingsAssetUserData.cpp
+++ b/Engine/Plugins/Media/MediaFrameworkUtilities/Source/MediaFrameworkUtilitiesEditor/Private/MediaFrameworkWorldSettingsAssetUserData.cpp
@@ -2,8 +2,9 @@
 
 #include "MediaFrameworkWorldSettingsAssetUserData.h"
 #include "GameFramework/Actor.h"
+#include "Serialization/CustomVersion.h"
 #include "UObject/EnterpriseObjectVersion.h"
-
+#include "UObject/UE5ReleaseStreamObjectVersion.h"
 
 FMediaFrameworkCaptureCurrentViewportOutputInfo::FMediaFrameworkCaptureCurrentViewportOutputInfo()
 	: MediaOutput(nullptr)
@@ -18,6 +19,30 @@ FMediaFrameworkCaptureCameraViewportCameraOutputInfo::FMediaFrameworkCaptureCame
 {
 }
 
+bool FMediaFrameworkCaptureCameraViewportCameraOutputInfo::Serialize(FArchive& Ar)
+{
+	Ar.UsingCustomVersion(FUE5ReleaseStreamObjectVersion::GUID);
+	return false;
+}
+
+#if WITH_EDITORONLY_DATA
+PRAGMA_DISABLE_DEPRECATION_WARNINGS
+void FMediaFrameworkCaptureCameraViewportCameraOutputInfo::PostSerialize(const FArchive& Ar)
+{
+	const int32 CustomVersion = Ar.CustomVer(FUE5ReleaseStreamObjectVersion::GUID);
+	if (Ar.IsLoading() && CustomVersion < FUE5ReleaseStreamObjectVersion::MediaProfilePluginCaptureCameraSoftPtr)
+	{
+
+		for (const TLazyObjectPtr<AActor>& LazyActor : LockedActors_DEPRECATED)
+		{
+			Cameras.Add(LazyActor.Get());
+		}
+
+		LockedActors_DEPRECATED.Empty();
+	}
+}
+PRAGMA_ENABLE_DEPRECATION_WARNINGS
+#endif
 
 FMediaFrameworkCaptureRenderTargetCameraOutputInfo::FMediaFrameworkCaptureRenderTargetCameraOutputInfo()
 	: RenderTarget(nullptr)
@@ -48,7 +73,7 @@ void UMediaFrameworkWorldSettingsAssetUserData::Serialize(FArchive& Ar)
 				{
 					if (Actor)
 					{
-						OutputInfo.LockedActors.Add(Actor);
+						OutputInfo.Cameras.Add(Actor);
 					}
 				}
 				OutputInfo.LockedCameraActors_DEPRECATED.Empty();
diff --git a/Engine/Plugins/Media/MediaFrameworkUtilities/Source/MediaFrameworkUtilitiesEditor/Private/MediaFrameworkWorldSettingsAssetUserData.h b/Engine/Plugins/Media/MediaFrameworkUtilities/Source/MediaFrameworkUtilitiesEditor/Private/MediaFrameworkWorldSettingsAssetUserData.h
--- a/Engine/Plugins/Media/MediaFrameworkUtilities/Source/MediaFrameworkUtilitiesEditor/Private/MediaFrameworkWorldSettingsAssetUserData.h
+++ b/Engine/Plugins/Media/MediaFrameworkUtilities/Source/MediaFrameworkUtilitiesEditor/Private/MediaFrameworkWorldSettingsAssetUserData.h
@@ -43,8 +43,22 @@ struct FMediaFrameworkCaptureCameraViewportCameraOutputInfo
 
 	FMediaFrameworkCaptureCameraViewportCameraOutputInfo();
 
-	UPROPERTY(EditAnywhere, Category="MediaViewportCapture", meta=(DisplayName="Cameras"))
-	TArray<TLazyObjectPtr<AActor>> LockedActors;
+	PRAGMA_DISABLE_DEPRECATION_WARNINGS
+	// Note: We need to explicitly disable warnings on these constructors/operators for clang to be happy with deprecated variables
+	// this is a requirement for clang to compile without warnings.
+	~FMediaFrameworkCaptureCameraViewportCameraOutputInfo() = default;
+	FMediaFrameworkCaptureCameraViewportCameraOutputInfo(const FMediaFrameworkCaptureCameraViewportCameraOutputInfo&) = default;
+	FMediaFrameworkCaptureCameraViewportCameraOutputInfo(FMediaFrameworkCaptureCameraViewportCameraOutputInfo&&) = default;
+	FMediaFrameworkCaptureCameraViewportCameraOutputInfo& operator=(const FMediaFrameworkCaptureCameraViewportCameraOutputInfo&) = default;
+	FMediaFrameworkCaptureCameraViewportCameraOutputInfo& operator=(FMediaFrameworkCaptureCameraViewportCameraOutputInfo&&) = default;
+	PRAGMA_ENABLE_DEPRECATION_WARNINGS
+	
+	UE_DEPRECATED(5.7, "Use Cameras instead")
+	UPROPERTY(Meta = (DeprecatedProperty, DeprecationMessage = "This property is no longer supported. Use Cameras instead"))
+	TArray<TLazyObjectPtr<AActor>> LockedActors_DEPRECATED;
+	
+	UPROPERTY(EditAnywhere, Category="MediaViewportCapture")
+	TArray<TSoftObjectPtr<AActor>> Cameras;
 
 	UPROPERTY(EditAnywhere, Category="MediaViewportCapture")
 	TObjectPtr<UMediaOutput> MediaOutput;
@@ -55,6 +69,12 @@ struct FMediaFrameworkCaptureCameraViewportCameraOutputInfo
 	UPROPERTY()
 	TEnumAsByte<EViewModeIndex> ViewMode;
 
+	bool Serialize(FArchive& Ar);
+	
+#if WITH_EDITORONLY_DATA
+	void PostSerialize(const FArchive& Ar);
+#endif
+	
 private:
 	//DEPRECATED 4.21 The type of LockedCameraActors has changed and will be removed from the code base in a future release. Use LockedActors.
 	UPROPERTY()
@@ -62,6 +82,17 @@ private:
 	friend UMediaFrameworkWorldSettingsAssetUserData;
 };
 
+template<> struct TStructOpsTypeTraits<FMediaFrameworkCaptureCameraViewportCameraOutputInfo> : public TStructOpsTypeTraitsBase2<FMediaFrameworkCaptureCameraViewportCameraOutputInfo>
+{
+	enum
+	{
+		WithSerializer = true,
+
+#if WITH_EDITORONLY_DATA
+		WithPostSerialize = true,
+#endif
+	};
+};
 
 /**
  * FMediaFrameworkCaptureRenderTargetCameraOutputInfo
diff --git a/Engine/Source/Runtime/Core/Public/UObject/UE5ReleaseStreamObjectVersions.inl b/Engine/Source/Runtime/Core/Public/UObject/UE5ReleaseStreamObjectVersions.inl
--- a/Engine/Source/Runtime/Core/Public/UObject/UE5ReleaseStreamObjectVersions.inl
+++ b/Engine/Source/Runtime/Core/Public/UObject/UE5ReleaseStreamObjectVersions.inl
@@ -183,5 +183,8 @@ UE5_RELEASE_VERSION(CompositePluginDerivedSceneCaptures, 57)
 // Add option to output attributes on the PCG Duplicate Cross Section node on the Data domain
 UE5_RELEASE_VERSION(ExtraOutputAttributesOnDataDomainPCG, 58)
 
+// Media Profile: Changed storage of capture cameras list from lazy to soft pointers
+UE5_RELEASE_VERSION(MediaProfilePluginCaptureCameraSoftPtr, 59)
+
 // -----<new versions can be added above this line>-------------------------------------------------
 #undef UE5_RELEASE_VERSION
\ No newline at end of file
