Media IO - Fix crash when refreshing media properties for Aja source

This was caused by restarting capture too quickly after stopping it. Media Profile code now uses the StopCapture_Callback method to wait until StopCapture is completely done before attempting to restart the capture with updated options.

diff --git a/Engine/Plugins/Media/AjaMedia/Source/Aja/Private/OutputChannel.cpp b/Engine/Plugins/Media/AjaMedia/Source/Aja/Private/OutputChannel.cpp
--- a/Engine/Plugins/Media/AjaMedia/Source/Aja/Private/OutputChannel.cpp
+++ b/Engine/Plugins/Media/AjaMedia/Source/Aja/Private/OutputChannel.cpp
@@ -1185,7 +1185,7 @@ namespace AJA
 
 			ThreadLock_PingPongOutputLoop_Memset0(CurrentFrame);
 		
-			AJA_CHECK(Device->WaitForInputOrOutputInterrupt(Channel));
+			Device->WaitForInputOrOutputInterrupt(Channel);
 
 			if (bTraceDebugMarkers)
 			{
@@ -1225,7 +1225,7 @@ namespace AJA
 							TextureTransfer->EndSync(CurrentFrame->VideoBuffer);
 						}
 					}
-					AJA_CHECK(Device->WaitForInputOrOutputInterrupt(Channel));
+					Device->WaitForInputOrOutputInterrupt(Channel);
 				}
 
 				ULWord VBICount = 0;
diff --git a/Engine/Plugins/Media/MediaFrameworkUtilities/Source/MediaFrameworkUtilities/Private/Profile/MediaProfilePlaybackManager.cpp b/Engine/Plugins/Media/MediaFrameworkUtilities/Source/MediaFrameworkUtilities/Private/Profile/MediaProfilePlaybackManager.cpp
--- a/Engine/Plugins/Media/MediaFrameworkUtilities/Source/MediaFrameworkUtilities/Private/Profile/MediaProfilePlaybackManager.cpp
+++ b/Engine/Plugins/Media/MediaFrameworkUtilities/Source/MediaFrameworkUtilities/Private/Profile/MediaProfilePlaybackManager.cpp
@@ -771,12 +771,7 @@ UMediaCapture* UMediaProfilePlaybackManager::OpenManagedViewportOutputFromIndex(
 #if WITH_EDITOR
 		if (bAutoRestartCapture)
 		{
-			FRecaptureFunc RestartCaptureFunc = [InMediaOutputIndex, InCaptureOptions, bAutoRestartCapture](UMediaProfilePlaybackManager* InPlaybackManager)
-			{
-				InPlaybackManager->OpenManagedViewportOutputFromIndex(InMediaOutputIndex, InCaptureOptions, bAutoRestartCapture);
-			};
-			
-			MediaOutput->OnOutputModified().AddUObject(this, &UMediaProfilePlaybackManager::QueueRestartCapture, RestartCaptureFunc);
+			MediaOutput->OnOutputModified().AddUObject(this, &UMediaProfilePlaybackManager::RestartManagedViewportOutput, InCaptureOptions, bAutoRestartCapture);
 		}
 #endif
 		
@@ -908,12 +903,7 @@ UMediaCapture* UMediaProfilePlaybackManager::OpenActiveViewportOutputFromIndex(i
 #if WITH_EDITOR
 		if (bAutoRestartCapture)
 		{
-			FRecaptureFunc RestartCaptureFunc = [InMediaOutputIndex, InCaptureOptions, bAutoRestartCapture](UMediaProfilePlaybackManager* InPlaybackManager)
-			{
-				InPlaybackManager->OpenActiveViewportOutputFromIndex(InMediaOutputIndex, InCaptureOptions, bAutoRestartCapture);
-			};
-			
-			MediaOutput->OnOutputModified().AddUObject(this, &UMediaProfilePlaybackManager::QueueRestartCapture, RestartCaptureFunc);
+			MediaOutput->OnOutputModified().AddUObject(this, &UMediaProfilePlaybackManager::RestartActiveViewportOutput, InCaptureOptions, bAutoRestartCapture);
 		}
 #endif
 		
@@ -967,12 +957,7 @@ UMediaCapture* UMediaProfilePlaybackManager::OpenRenderTargetOutputFromIndex(int
 #if WITH_EDITOR
 		if (bAutoRestartCapture)
 		{
-			FRecaptureFunc RestartCaptureFunc = [InMediaOutputIndex, InRenderTarget, InCaptureOptions, bAutoRestartCapture](UMediaProfilePlaybackManager* InPlaybackManager)
-			{
-				InPlaybackManager->OpenRenderTargetOutputFromIndex(InMediaOutputIndex, InRenderTarget, InCaptureOptions, bAutoRestartCapture);
-			};
-			
-			MediaOutput->OnOutputModified().AddUObject(this, &UMediaProfilePlaybackManager::QueueRestartCapture, RestartCaptureFunc);
+			MediaOutput->OnOutputModified().AddUObject(this, &UMediaProfilePlaybackManager::RestartRenderTargetOutput, InRenderTarget, InCaptureOptions, bAutoRestartCapture);
 		}
 #endif
 		
@@ -994,6 +979,53 @@ UMediaCapture* UMediaProfilePlaybackManager::OpenRenderTargetOutput(UMediaOutput
 	return OpenRenderTargetOutputFromIndex(MediaOutputIndex, InRenderTarget, InCaptureOptions, bAutoRestartCapture);
 }
 
+void UMediaProfilePlaybackManager::RestartActiveViewportOutput(UMediaOutput* InMediaOutput, FMediaCaptureOptions InCaptureOptions, bool bAutoRestartCapture)
+{
+	TWeakObjectPtr<UMediaOutput> WeakMediaOutput = InMediaOutput;
+	auto Callback = [WeakMediaOutput, InCaptureOptions, bAutoRestartCapture](UMediaProfilePlaybackManager* Manager)
+	{
+		TStrongObjectPtr<UMediaOutput> MediaOutput = WeakMediaOutput.Pin();
+
+		if (Manager && MediaOutput)
+		{
+			Manager->OpenActiveViewportOutput(MediaOutput.Get(), InCaptureOptions, bAutoRestartCapture);
+		}
+	};
+	QueueRestartCapture(InMediaOutput, Callback);
+}
+
+void UMediaProfilePlaybackManager::RestartManagedViewportOutput(UMediaOutput* InMediaOutput, FMediaCaptureOptions InCaptureOptions, bool bAutoRestartCapture)
+{
+	TWeakObjectPtr<UMediaOutput> WeakMediaOutput = InMediaOutput;
+	auto Callback = [WeakMediaOutput, InCaptureOptions, bAutoRestartCapture](UMediaProfilePlaybackManager* Manager)
+	{
+		TStrongObjectPtr<UMediaOutput> MediaOutput = WeakMediaOutput.Pin();
+
+		if (Manager && MediaOutput)
+		{
+			Manager->OpenManagedViewportOutput(MediaOutput.Get(), InCaptureOptions, bAutoRestartCapture);
+		}
+	};
+	QueueRestartCapture(InMediaOutput, Callback);
+}
+
+void UMediaProfilePlaybackManager::RestartRenderTargetOutput(UMediaOutput* InMediaOutput, UTextureRenderTarget2D* InRenderTarget, FMediaCaptureOptions InCaptureOptions, bool bAutoRestartCapture)
+{
+	TWeakObjectPtr<UTextureRenderTarget2D> WeakRenderTarget = InRenderTarget;
+	TWeakObjectPtr<UMediaOutput> WeakMediaOutput = InMediaOutput;
+	auto Callback = [WeakRenderTarget, WeakMediaOutput, InCaptureOptions, bAutoRestartCapture](UMediaProfilePlaybackManager* Manager)
+	{
+		TStrongObjectPtr<UTextureRenderTarget2D> RenderTarget = WeakRenderTarget.Pin();
+		TStrongObjectPtr<UMediaOutput> MediaOutput = WeakMediaOutput.Pin();
+
+		if (Manager && RenderTarget && MediaOutput)
+		{
+			Manager->OpenRenderTargetOutput(MediaOutput.Get(), RenderTarget.Get(), InCaptureOptions, bAutoRestartCapture);
+		}
+	};
+	QueueRestartCapture(InMediaOutput, Callback);
+}
+
 bool UMediaProfilePlaybackManager::IsOutputCapturingFromIndex(int32 InMediaOutputIndex) const
 {
 	UMediaProfile* MediaProfile = GetOwningMediaProfile();
@@ -1065,10 +1097,10 @@ TOptional<EMediaCaptureState> UMediaProfilePlaybackManager::GetOutputCaptureStat
 	return GetOutputCaptureStateFromIndex(MediaOutputIndex, bOutHasError);
 }
 
-void UMediaProfilePlaybackManager::CloseOutputFromIndex(int32 InMediaOutputIndex, bool bDestroyCaptureObjects)
+void UMediaProfilePlaybackManager::CloseOutputFromIndex(int32 InMediaOutputIndex, const FCloseOutputArgs& InArgs)
 {
 	using namespace MediaProfilePlaybackManager;
-	
+
 	UMediaProfile* MediaProfile = GetOwningMediaProfile();
 	if (!MediaProfile)
 	{
@@ -1080,7 +1112,7 @@ void UMediaProfilePlaybackManager::CloseOutputFromIndex(int32 InMediaOutputIndex
 	{
 		return;
 	}
-	
+
 	if (!MediaOutputCaptures.Contains(InMediaOutputIndex))
 	{
 		return;
@@ -1089,16 +1121,38 @@ void UMediaProfilePlaybackManager::CloseOutputFromIndex(int32 InMediaOutputIndex
 	UMediaCapture* MediaCapture = MediaOutputCaptures[InMediaOutputIndex];
 
 	constexpr bool bAllowPendingFrameToBeProcessed = false;
-	MediaCapture->StopCapture(bAllowPendingFrameToBeProcessed);
+	MediaCapture->StopCapture_Callback(bAllowPendingFrameToBeProcessed, FSimpleDelegate::CreateUObject(this, &UMediaProfilePlaybackManager::PostStopCapture, InMediaOutputIndex, InArgs));
+}
 
-	if (bDestroyCaptureObjects)
+void UMediaProfilePlaybackManager::PostStopCapture(int32 InMediaOutputIndex, FCloseOutputArgs InArgs)
+{
+	UMediaProfile* MediaProfile = GetOwningMediaProfile();
+	if (!MediaProfile)
+	{
+		return;
+	}
+
+	UMediaOutput* MediaOutput = MediaProfile->GetMediaOutput(InMediaOutputIndex);
+	if (!IsValid(MediaOutput))
+	{
+		return;
+	}
+
+	if (!MediaOutputCaptures.Contains(InMediaOutputIndex))
+	{
+		return;
+	}
+
+	UMediaCapture* MediaCapture = MediaOutputCaptures[InMediaOutputIndex];
+
+	if (InArgs.bDestroyCaptureObjects)
 	{
-		ManagedOutputViewports.RemoveAll([MediaOutput](const TSharedPtr<FManagedViewport>& InManagedViewport)
+		ManagedOutputViewports.RemoveAll([MediaOutput](const TSharedPtr<MediaProfilePlaybackManager::FManagedViewport>& InManagedViewport)
 		{
 			return InManagedViewport->MediaOutput == MediaOutput && InManagedViewport->Consumers.IsEmpty();
 		});
 
-		ActiveOutputViewports.RemoveAll([MediaOutput](const TSharedPtr<FActiveViewport>& InViewport)
+		ActiveOutputViewports.RemoveAll([MediaOutput](const TSharedPtr<MediaProfilePlaybackManager::FActiveViewport>& InViewport)
 		{
 			return InViewport->MediaOutput == MediaOutput;
 		});
@@ -1118,7 +1172,7 @@ void UMediaProfilePlaybackManager::CloseOutputFromIndex(int32 InMediaOutputIndex
 			}
 		}
 #endif
-		
+
 		MediaCapture->OnStateChangedNative.RemoveAll(this);
 		MediaOutputCaptures.Remove(InMediaOutputIndex);
 	}
@@ -1126,9 +1180,11 @@ void UMediaProfilePlaybackManager::CloseOutputFromIndex(int32 InMediaOutputIndex
 #if WITH_EDITOR
 	MediaOutput->OnOutputModified().RemoveAll(this);
 #endif
+
+	InArgs.Callback.ExecuteIfBound();
 }
 
-void UMediaProfilePlaybackManager::CloseOutput(UMediaOutput* InMediaOutput, bool bDestroyCaptureObjects)
+void UMediaProfilePlaybackManager::CloseOutput(UMediaOutput* InMediaOutput, const FCloseOutputArgs& InArgs)
 {
 	UMediaProfile* MediaProfile = GetOwningMediaProfile();
 	if (!MediaProfile)
@@ -1137,7 +1193,7 @@ void UMediaProfilePlaybackManager::CloseOutput(UMediaOutput* InMediaOutput, bool
 	}
 	
 	int32 MediaOutputIndex = MediaProfile->FindMediaOutputIndex(InMediaOutput);
-	CloseOutputFromIndex(MediaOutputIndex, bDestroyCaptureObjects);
+	CloseOutputFromIndex(MediaOutputIndex, InArgs);
 }
 
 UMediaPlayer* UMediaProfilePlaybackManager::FindOrCreateMediaPlayerForIndex(int32 InMediaSourceIndex)
@@ -1371,11 +1427,21 @@ void UMediaProfilePlaybackManager::OnCaptureStateChanged(UMediaCapture* InMediaC
 
 void UMediaProfilePlaybackManager::QueueRestartCapture(UMediaOutput* InMediaOutput, FRecaptureFunc InRestartCaptureFunc)
 {
-	constexpr bool bDestroyCaptureObjects = false;
-	CloseOutput(InMediaOutput, bDestroyCaptureObjects);
-	
 	RestartCaptureFuncs.Add(InRestartCaptureFunc);
-	
+
+	constexpr bool bDestroyCaptureObjects = false;
+
+	FCloseOutputArgs Args;
+	Args.bDestroyCaptureObjects = false;
+	Args.Callback = FSimpleDelegate::CreateUObject(this, &UMediaProfilePlaybackManager::DeferredRestart);
+
+	CloseOutput(InMediaOutput, MoveTemp(Args));
+}
+
+void UMediaProfilePlaybackManager::DeferredRestart()
+{
+	// Not all implementations of MediaCapture implementations have a StopCapture_Future method, 
+	// so we defer to the next tick to make sure those MediaCaptures also have a chance to shutdown before he output is restarted.
 	if (!RestartCapturesTimerHandle.IsValid())
 	{
 #if WITH_EDITOR
@@ -1402,7 +1468,7 @@ void UMediaProfilePlaybackManager::RestartCaptures()
 	RestartCaptureFuncs.Empty();
 	RestartCapturesTimerHandle.Invalidate();
 }
-
+ 
 void UMediaProfilePlaybackManager::Cleanup()
 {
 	TArray<int32> OpenSources;
diff --git a/Engine/Plugins/Media/MediaFrameworkUtilities/Source/MediaFrameworkUtilities/Public/Profile/MediaProfilePlaybackManager.h b/Engine/Plugins/Media/MediaFrameworkUtilities/Source/MediaFrameworkUtilities/Public/Profile/MediaProfilePlaybackManager.h
--- a/Engine/Plugins/Media/MediaFrameworkUtilities/Source/MediaFrameworkUtilities/Public/Profile/MediaProfilePlaybackManager.h
+++ b/Engine/Plugins/Media/MediaFrameworkUtilities/Source/MediaFrameworkUtilities/Public/Profile/MediaProfilePlaybackManager.h
@@ -233,7 +233,32 @@ public:
 	 * @return The media capture object, if the capture was started successfully; otherwise, null
 	 */
 	UMediaCapture* OpenRenderTargetOutput(UMediaOutput* InMediaOutput, UTextureRenderTarget2D* InRenderTarget, const FMediaCaptureOptions& InCaptureOptions, bool bAutoRestartCapture = false);
-	
+
+	 /* 
+	 * Restarts a capture of an existing engine viewport for a specified media output
+	 * @param InMediaOutput The media output in the media profile
+	 * @param InCaptureOptions The capture options to capture with
+	 * @param bAutoRestartCapture Indicates that this capture will automatically restart if any modifications to the media output are detected
+	 */
+	void RestartActiveViewportOutput(UMediaOutput* InMediaOutput, FMediaCaptureOptions InCaptureOptions, bool bAutoRestartCapture = false);
+
+	/**
+	 * Restarts a capture of a managed viewport for a specified media output
+	 * @param InMediaOutput The media output in the media profile
+	 * @param InCaptureOptions The capture options to capture with
+	 * @param bAutoRestartCapture Indicates that this capture will automatically restart if any modifications to the media output are detected
+	 */
+	void RestartManagedViewportOutput(UMediaOutput* InMediaOutput, FMediaCaptureOptions InCaptureOptions, bool bAutoRestartCapture = false);
+
+	/**
+	 * Restarts a capture of a render target for a specified media output
+	 * @param InMediaOutput The media output in the media profile
+	 * @param InRenderTarget The render target to capture
+	 * @param InCaptureOptions The capture options to capture with
+	 * @param bAutoRestartCapture Indicates that this capture will automatically restart if any modifications to the media output are detected
+	 */
+	void RestartRenderTargetOutput(UMediaOutput* InMediaOutput, UTextureRenderTarget2D* InRenderTarget, FMediaCaptureOptions InCaptureOptions, bool bAutoRestartCapture = false);
+
 	/**
 	 * Gets whether there is an active capture for the specified media output index
 	 * @param InMediaOutputIndex The index of the media output within the media profile to check
@@ -264,21 +289,34 @@ public:
 	 */
 	TOptional<EMediaCaptureState> GetOutputCaptureState(UMediaOutput* InMediaOutput, bool& bOutHasError) const;
 	
+	/** Arguments that control what happens when a media output is closed */
+	struct FCloseOutputArgs
+	{
+		/** Indicates that the media capture and all related objects should be destroyed. Set to false if the capture is immediately going to be restarted. */
+		bool bDestroyCaptureObjects;
+
+		/** Some capture implementations can take a few frames to completely stop. When they are completely done, this callback delegate will be invoked. */
+		FSimpleDelegate Callback;
+
+		FCloseOutputArgs()
+			: bDestroyCaptureObjects(true)
+		{
+		}
+	};
+
 	/**
 	 * Closes the capture of the specified media output
-	 * @param InMediaOutputIndex The index of the media output in the media profile to close
-	 * @param bDestroyCaptureObjects Indicates that the media capture and all related objects should be destroyed.
-	 *								 Set to false if the capture is immediately going to be restarted.
+	 * @param InMediaOutputIndex The index of the media output within the media profile
+	 * @param InArgs Arguments that control what happens when the output is stopped.
 	 */
-	void CloseOutputFromIndex(int32 InMediaOutputIndex, bool bDestroyCaptureObjects = true);
+	void CloseOutputFromIndex(int32 InMediaOutputIndex, const FCloseOutputArgs& InArgs = FCloseOutputArgs());
 
 	/**
 	 * Closes the capture of the specified media output
 	 * @param InMediaOutput The media output in the media profile to close
-	 * @param bDestroyCaptureObjects Indicates that the media capture and all related objects should be destroyed.
-	 *								 Set to false if the capture is immediately going to be restarted.
+	 * @param InArgs Arguments that control what happens when the output is stopped.
 	 */
-	void CloseOutput(UMediaOutput* InMediaOutput, bool bDestroyCaptureObjects = true);
+	void CloseOutput(UMediaOutput* InMediaOutput, const FCloseOutputArgs& InArgs = FCloseOutputArgs());
 	
 private:
 	/**
@@ -323,11 +361,17 @@ private:
 	/** Raised when a change to the specified media output has been detected and the capture shout be auto-restarted */
 	void QueueRestartCapture(UMediaOutput* InMediaOutput, FRecaptureFunc InRestartCaptureFunc);
 
+	/** Invoked through QueueRestartCapture when the StopCapture implementation is done. */
+	void DeferredRestart();
+
 	/** Called on the next available tick, restarts any captures that have been queued to auto-restart */
 	void RestartCaptures();
 
 	/** Cleans up all open sources and outputs, and deletes any managed state */
 	void Cleanup();
+
+	/** Handles cleaning up resources and invoking the post StopCapture callback. */
+	void PostStopCapture(int32 InMediaOutputIndex, FCloseOutputArgs InArgs);
 	
 #if WITH_EDITOR
 	/** Raised when a PIE session has started */
diff --git a/Engine/Plugins/Media/MediaFrameworkUtilities/Source/MediaFrameworkUtilitiesEditor/Private/AssetEditor/MediaProfileEditor.cpp b/Engine/Plugins/Media/MediaFrameworkUtilities/Source/MediaFrameworkUtilitiesEditor/Private/AssetEditor/MediaProfileEditor.cpp
--- a/Engine/Plugins/Media/MediaFrameworkUtilities/Source/MediaFrameworkUtilitiesEditor/Private/AssetEditor/MediaProfileEditor.cpp
+++ b/Engine/Plugins/Media/MediaFrameworkUtilities/Source/MediaFrameworkUtilitiesEditor/Private/AssetEditor/MediaProfileEditor.cpp
@@ -383,21 +383,18 @@ void FMediaProfileEditor::RefreshSelectedMediaItems(const TArray<int32>& InMedia
 		{
 			if (PlaybackManager->IsOutputCapturing(MediaOutput))
 			{
-				constexpr bool bDestroyCaptureObjects = false;
-				PlaybackManager->CloseOutput(MediaOutput, bDestroyCaptureObjects);
-
 				if (UMediaProfileEditorCaptureSettings* CaptureSettings = GetMediaFrameworkCaptureSettings())
 				{
 					if (CaptureSettings->CurrentViewportMediaOutput.MediaOutput == MediaOutput)
 					{
-						PlaybackManager->OpenActiveViewportOutput(MediaOutput, CaptureSettings->CurrentViewportMediaOutput.CaptureOptions, CaptureSettings->bAutoRestartCaptureOnChange);
+						PlaybackManager->RestartActiveViewportOutput(MediaOutput, CaptureSettings->CurrentViewportMediaOutput.CaptureOptions, CaptureSettings->bAutoRestartCaptureOnChange);
 					}
 
 					for (const FMediaFrameworkCaptureCameraViewportCameraOutputInfo& ViewportCapture : CaptureSettings->ViewportCaptures)
 					{
 						if (ViewportCapture.MediaOutput == MediaOutput)
 						{
-							PlaybackManager->OpenManagedViewportOutput(MediaOutput, ViewportCapture.CaptureOptions, CaptureSettings->bAutoRestartCaptureOnChange);
+							PlaybackManager->RestartManagedViewportOutput(MediaOutput, ViewportCapture.CaptureOptions, CaptureSettings->bAutoRestartCaptureOnChange);
 						}
 					}
 
@@ -405,7 +402,7 @@ void FMediaProfileEditor::RefreshSelectedMediaItems(const TArray<int32>& InMedia
 					{
 						if (RenderTargetCapture.MediaOutput == MediaOutput)
 						{
-							PlaybackManager->OpenRenderTargetOutput(MediaOutput, RenderTargetCapture.RenderTarget, RenderTargetCapture.CaptureOptions, CaptureSettings->bAutoRestartCaptureOnChange);
+							PlaybackManager->RestartRenderTargetOutput(MediaOutput, RenderTargetCapture.RenderTarget, RenderTargetCapture.CaptureOptions, CaptureSettings->bAutoRestartCaptureOnChange);
 						}
 					}
 				}
@@ -564,21 +561,18 @@ void FMediaProfileEditor::OnObjectPropertyChanged(UObject* InObject, FPropertyCh
 			{
 				if (PlaybackManager->IsOutputCapturing(MediaOutput))
 				{
-					constexpr bool bDestroyCaptureObjects = false;
-					PlaybackManager->CloseOutput(MediaOutput, bDestroyCaptureObjects);
-
 					if (UMediaProfileEditorCaptureSettings* CaptureSettings = GetMediaFrameworkCaptureSettings())
 					{
 						if (CaptureSettings->CurrentViewportMediaOutput.MediaOutput == MediaOutput)
 						{
-							PlaybackManager->OpenActiveViewportOutput(MediaOutput, CaptureSettings->CurrentViewportMediaOutput.CaptureOptions, CaptureSettings->bAutoRestartCaptureOnChange);
+							PlaybackManager->RestartActiveViewportOutput(MediaOutput, CaptureSettings->CurrentViewportMediaOutput.CaptureOptions, CaptureSettings->bAutoRestartCaptureOnChange);
 						}
 
 						for (const FMediaFrameworkCaptureCameraViewportCameraOutputInfo& ViewportCapture : CaptureSettings->ViewportCaptures)
 						{
 							if (ViewportCapture.MediaOutput == MediaOutput)
 							{
-								PlaybackManager->OpenManagedViewportOutput(MediaOutput, ViewportCapture.CaptureOptions, CaptureSettings->bAutoRestartCaptureOnChange);
+								PlaybackManager->RestartManagedViewportOutput(MediaOutput, ViewportCapture.CaptureOptions, CaptureSettings->bAutoRestartCaptureOnChange);
 							}
 						}
 
@@ -586,7 +580,7 @@ void FMediaProfileEditor::OnObjectPropertyChanged(UObject* InObject, FPropertyCh
 						{
 							if (RenderTargetCapture.MediaOutput == MediaOutput)
 							{
-								PlaybackManager->OpenRenderTargetOutput(MediaOutput, RenderTargetCapture.RenderTarget, RenderTargetCapture.CaptureOptions, CaptureSettings->bAutoRestartCaptureOnChange);
+								PlaybackManager->RestartRenderTargetOutput(MediaOutput, RenderTargetCapture.RenderTarget, RenderTargetCapture.CaptureOptions, CaptureSettings->bAutoRestartCaptureOnChange);
 							}
 						}
 					}
diff --git a/Engine/Plugins/Media/MediaIOFramework/Source/MediaIOCore/Private/MediaCapture.cpp b/Engine/Plugins/Media/MediaIOFramework/Source/MediaIOCore/Private/MediaCapture.cpp
--- a/Engine/Plugins/Media/MediaIOFramework/Source/MediaIOCore/Private/MediaCapture.cpp
+++ b/Engine/Plugins/Media/MediaIOFramework/Source/MediaIOCore/Private/MediaCapture.cpp
@@ -548,6 +548,11 @@ bool UMediaCapture::UpdateTextureRenderTarget2D(UTextureRenderTarget2D* InRender
 }
 
 void UMediaCapture::StopCapture(bool bAllowPendingFrameToBeProcess)
+{
+	StopCapture_Callback(bAllowPendingFrameToBeProcess, FSimpleDelegate{});
+}
+
+void UMediaCapture::StopCapture_Callback(bool bAllowPendingFrameToBeProcess, FSimpleDelegate Callback)
 {
 	if (bWaitingForStopCapture)
 	{
@@ -558,12 +563,12 @@ void UMediaCapture::StopCapture(bool bAllowPendingFrameToBeProcess)
 	if (!IsInGameThread())
 	{
 		TWeakObjectPtr<UMediaCapture> Self = this;
-		AsyncTask(ENamedThreads::GameThread, [Self, bAllowPendingFrameToBeProcess]()
+		AsyncTask(ENamedThreads::GameThread, [Self, bAllowPendingFrameToBeProcess, Callback]()
 			{
 				UMediaCapture* MediaCapture = Self.Get();
 				if (UObjectInitialized() && MediaCapture)
 				{
-					MediaCapture->StopCapture(bAllowPendingFrameToBeProcess);
+					MediaCapture->StopCapture_Callback(bAllowPendingFrameToBeProcess, Callback);
 				}
 			});
 
@@ -625,9 +630,9 @@ void UMediaCapture::StopCapture(bool bAllowPendingFrameToBeProcess)
 				TWeakObjectPtr<UMediaCapture> Self = this;
 
 				TRACE_BOOKMARK(TEXT("UMediaCapture::InitiateStopCapture"));
-				StopCapture_Future().Then([Self](TFuture<void>) mutable
+				StopCapture_Future().Then([Self, Callback](TFuture<void>) mutable
 				{
-					AsyncTask(ENamedThreads::GameThread, [Self]
+					AsyncTask(ENamedThreads::GameThread, [Self, Callback]
 					{
 						// Setting bEvenifPendingKill because this might be called as part of BeginDestroy,
 						// but we still want to cleanup card connections.
@@ -635,6 +640,7 @@ void UMediaCapture::StopCapture(bool bAllowPendingFrameToBeProcess)
 						{
 							Capture->SetState(EMediaCaptureState::Stopped);
 							Capture->PostCaptureCleanup();
+							Callback.ExecuteIfBound();
 							TRACE_BOOKMARK(TEXT("UMediaCapture::CompleteStopCapture"));
 						}
 					});
@@ -644,8 +650,13 @@ void UMediaCapture::StopCapture(bool bAllowPendingFrameToBeProcess)
 			{
 				StopCaptureImpl(false);
 				PostCaptureCleanup();
+				Callback.ExecuteIfBound();
 			}
 		}
+		else
+		{
+			Callback.ExecuteIfBound();
+		}
 	}
 }
 
diff --git a/Engine/Plugins/Media/MediaIOFramework/Source/MediaIOCore/Public/MediaCapture.h b/Engine/Plugins/Media/MediaIOFramework/Source/MediaIOCore/Public/MediaCapture.h
--- a/Engine/Plugins/Media/MediaIOFramework/Source/MediaIOCore/Public/MediaCapture.h
+++ b/Engine/Plugins/Media/MediaIOFramework/Source/MediaIOCore/Public/MediaCapture.h
@@ -351,6 +351,13 @@ public:
 	UFUNCTION(BlueprintCallable, Category = "Media|Output")
 	UE_API void StopCapture(bool bAllowPendingFrameToBeProcess);
 
+	/**
+	 * Implementation of StopCapture that allows providing a callback that's invoked when the capture has completely stopped.
+	 * Some capture implementations (ie. Aja), need a few frames to completely stop the capture, so this method can be used
+	 * to act when the underlying SDK is done shutting down.
+	 */
+	UE_API void StopCapture_Callback(bool bAllowPendingFrameToBeProcess, FSimpleDelegate Callback);
+
 	/** Get the current state of the capture. */
 	UFUNCTION(BlueprintCallable, Category = "Media|Output")
 	virtual EMediaCaptureState GetState() const { return MediaState; }
