Media Profile: Fixed issue where active media sources would get stopped whenever PIE was exited

diff --git a/Engine/Plugins/Media/MediaFrameworkUtilities/Source/MediaFrameworkUtilities/Private/Profile/MediaProfilePlaybackManager.cpp b/Engine/Plugins/Media/MediaFrameworkUtilities/Source/MediaFrameworkUtilities/Private/Profile/MediaProfilePlaybackManager.cpp
--- a/Engine/Plugins/Media/MediaFrameworkUtilities/Source/MediaFrameworkUtilities/Private/Profile/MediaProfilePlaybackManager.cpp
+++ b/Engine/Plugins/Media/MediaFrameworkUtilities/Source/MediaFrameworkUtilities/Private/Profile/MediaProfilePlaybackManager.cpp
@@ -1212,6 +1212,11 @@ UMediaPlayer* UMediaProfilePlaybackManager::FindOrCreateMediaPlayerForIndex(int3
 	if (!MediaSourcePlayers.Contains(InMediaSourceIndex))
 	{
 		UMediaPlayer* NewMediaPlayer = NewObject<UMediaPlayer>(this);
+
+#if WITH_EDITOR
+		NewMediaPlayer->AffectedByPIEHandling = false;
+#endif
+		
 		MediaSourcePlayers.Add(InMediaSourceIndex, NewMediaPlayer);
 	}
 
