Media Profile: Fix for crash caused by attempting to capture from camera spawned in by PIE after PIE session is ended and restarted

diff --git a/Engine/Plugins/Media/MediaFrameworkUtilities/Source/MediaFrameworkUtilitiesEditor/Private/AssetEditor/SMediaProfileMediaItemDisplay.cpp b/Engine/Plugins/Media/MediaFrameworkUtilities/Source/MediaFrameworkUtilitiesEditor/Private/AssetEditor/SMediaProfileMediaItemDisplay.cpp
--- a/Engine/Plugins/Media/MediaFrameworkUtilities/Source/MediaFrameworkUtilitiesEditor/Private/AssetEditor/SMediaProfileMediaItemDisplay.cpp
+++ b/Engine/Plugins/Media/MediaFrameworkUtilities/Source/MediaFrameworkUtilitiesEditor/Private/AssetEditor/SMediaProfileMediaItemDisplay.cpp
@@ -541,6 +541,13 @@ public:
 		AActor* ActorToLockTo = nullptr;
 		if (LockedActors.IsValidIndex(CurrentLockedActor))
 		{
+			if (!LockedActors[CurrentLockedActor].IsValid())
+			{
+				// If the locked actor is no longer valid, abort and set the current locked actor to none
+				CurrentLockedActor = INDEX_NONE;
+				return;
+			}
+			
 			AActor* Actor = LockedActors[CurrentLockedActor].Get();
 			ActorToLockTo = Actor;
 			
