MediaProfile - Allow autosaving changes whenever a property is modified

#rnx

diff --git a/Engine/Plugins/Media/MediaFrameworkUtilities/Source/MediaFrameworkUtilities/Public/Profile/MediaProfileSettings.h b/Engine/Plugins/Media/MediaFrameworkUtilities/Source/MediaFrameworkUtilities/Public/Profile/MediaProfileSettings.h
--- a/Engine/Plugins/Media/MediaFrameworkUtilities/Source/MediaFrameworkUtilities/Public/Profile/MediaProfileSettings.h
+++ b/Engine/Plugins/Media/MediaFrameworkUtilities/Source/MediaFrameworkUtilities/Public/Profile/MediaProfileSettings.h
@@ -34,6 +34,14 @@ public:
 	UPROPERTY(Config, EditAnywhere, Category="MediaProfile")
 	bool bApplyInCommandlet;
 
+#if WITH_EDITORONLY_DATA
+	/** 
+	 * Forces the media profile to save itself after a property change (if the profile was already saved once)
+	 */
+	UPROPERTY(Config)
+	bool bAutosaveAfterChange = false;
+#endif
+
 private:
 
 	UPROPERTY(AdvancedDisplay, Config, EditAnywhere, Category="MediaProfile")
diff --git a/Engine/Plugins/Media/MediaFrameworkUtilities/Source/MediaFrameworkUtilitiesEditor/Private/AssetEditor/MediaProfileEditor.cpp b/Engine/Plugins/Media/MediaFrameworkUtilities/Source/MediaFrameworkUtilitiesEditor/Private/AssetEditor/MediaProfileEditor.cpp
--- a/Engine/Plugins/Media/MediaFrameworkUtilities/Source/MediaFrameworkUtilitiesEditor/Private/AssetEditor/MediaProfileEditor.cpp
+++ b/Engine/Plugins/Media/MediaFrameworkUtilities/Source/MediaFrameworkUtilitiesEditor/Private/AssetEditor/MediaProfileEditor.cpp
@@ -587,7 +587,15 @@ void FMediaProfileEditor::OnObjectPropertyChanged(UObject* InObject, FPropertyCh
 				}
 			}
 		}
-	}	
+	}
+
+	if (GetDefault<UMediaProfileSettings>()->bAutosaveAfterChange)
+	{
+		if (MediaProfileBeingEdited && MediaProfileBeingEdited->GetPackage() != GetTransientPackage())
+		{
+			SaveAsset_Execute();
+		}
+	}
 }
 
 #undef LOCTEXT_NAMESPACE
