MediaProfile - Fix media profile autosaving regardless of asset being edited

diff --git a/Engine/Plugins/Media/MediaFrameworkUtilities/Source/MediaFrameworkUtilitiesEditor/Private/AssetEditor/MediaProfileEditor.cpp b/Engine/Plugins/Media/MediaFrameworkUtilities/Source/MediaFrameworkUtilitiesEditor/Private/AssetEditor/MediaProfileEditor.cpp
--- a/Engine/Plugins/Media/MediaFrameworkUtilities/Source/MediaFrameworkUtilitiesEditor/Private/AssetEditor/MediaProfileEditor.cpp
+++ b/Engine/Plugins/Media/MediaFrameworkUtilities/Source/MediaFrameworkUtilitiesEditor/Private/AssetEditor/MediaProfileEditor.cpp
@@ -589,11 +589,14 @@ void FMediaProfileEditor::OnObjectPropertyChanged(UObject* InObject, FPropertyCh
 		}
 	}
 
-	if (GetDefault<UMediaProfileSettings>()->bAutosaveAfterChange)
+	if (MediaProfileBeingEdited == InObject)
 	{
-		if (MediaProfileBeingEdited && MediaProfileBeingEdited->GetPackage() != GetTransientPackage())
+		if (GetDefault<UMediaProfileSettings>()->bAutosaveAfterChange)
 		{
-			SaveAsset_Execute();
+			if (MediaProfileBeingEdited && MediaProfileBeingEdited->GetPackage() != GetTransientPackage())
+			{
+				SaveAsset_Execute();
+			}
 		}
 	}
 }
