Media Profile - Allow creating a media profile button outside of the Level Editor toolbar

#rnx

diff --git a/Engine/Plugins/Media/MediaFrameworkUtilities/Source/MediaFrameworkUtilitiesEditor/Private/UI/MediaProfileMenuEntry.cpp b/Engine/Plugins/Media/MediaFrameworkUtilities/Source/MediaFrameworkUtilitiesEditor/Private/UI/MediaProfileMenuEntry.cpp
--- a/Engine/Plugins/Media/MediaFrameworkUtilities/Source/MediaFrameworkUtilitiesEditor/Private/UI/MediaProfileMenuEntry.cpp
+++ b/Engine/Plugins/Media/MediaFrameworkUtilities/Source/MediaFrameworkUtilitiesEditor/Private/UI/MediaProfileMenuEntry.cpp
@@ -11,6 +11,7 @@
 #include "Framework/Commands/UICommandList.h"
 #include "IAssetTools.h"
 #include "Kismet2/SClassPickerDialog.h"
+#include "Misc/ConfigCacheIni.h"
 #include "Misc/FeedbackContext.h"
 #include "Profile/IMediaProfileManager.h"
 #include "Profile/MediaProfile.h"
@@ -46,8 +47,31 @@ struct FMediaProfileMenuEntryImpl
 			})
 		);
 
-		UToolMenu* Menu = UToolMenus::Get()->ExtendMenu("LevelEditor.LevelEditorToolBar.User");
+
+		FString ButtonAlignmentString;
+
+		/** 
+		 * These settings aren't meant to be user-configurable, they're meant to allow an application (Currently LiveLinkHub)
+		 * to provide an override to the default location and alignment of the widget. This ultimately allows hosting the MediaProfile toolbar button
+		 * in a different asset editor than the Level Editor.
+		 */
+		GConfig->GetString(TEXT("MediaProfile"), TEXT("ButtonMenu"), ToolbarMenu, GEngineIni);
+		GConfig->GetString(TEXT("MediaProfile"), TEXT("ButtonAlignment"), ButtonAlignmentString, GEngineIni);
+
+		if (ToolbarMenu.IsEmpty())
+		{
+			ToolbarMenu = TEXT("LevelEditor.LevelEditorToolBar.User");
+		}
+
+		EToolMenuSectionAlign ButtonAlignment = EToolMenuSectionAlign::Default;
+		if (!ButtonAlignmentString.IsEmpty())
+		{
+			ButtonAlignment = static_cast<EToolMenuSectionAlign>(StaticEnum<EToolMenuSectionAlign>()->GetValueByName(*ButtonAlignmentString));
+		}
+
+		UToolMenu* Menu = UToolMenus::Get()->ExtendMenu(*ToolbarMenu);
 		FToolMenuSection& Section = Menu->FindOrAddSection("MediaProfile");
+		Section.Alignment = ButtonAlignment;
 
 		auto ButtonTooltipLambda = [this]() -> FText
 		{
@@ -68,6 +92,7 @@ struct FMediaProfileMenuEntryImpl
 			FSlateIcon(FMediaFrameworkUtilitiesEditorStyle::GetStyleSetName(), TEXT("ToolbarIcon.Edit"))
 		);
 		MediaProfileButtonEntry.SetCommandList(Actions);
+		MediaProfileButtonEntry.ToolBarData.LabelOverride = FText::GetEmpty();
 
 		FToolMenuEntry MediaProfileMenuEntry = FToolMenuEntry::InitComboButton(
 		"MediaProfileMenu",
@@ -88,7 +113,7 @@ struct FMediaProfileMenuEntryImpl
 	{
 		if (!IsEngineExitRequested() && UObjectInitialized())
 		{
-			UToolMenus::Get()->RemoveSection("LevelEditor.LevelEditorToolBar.User", "MediaProfile");
+			UToolMenus::Get()->RemoveSection(*ToolbarMenu, "MediaProfile");
 		}
 	}
 
@@ -189,7 +214,8 @@ struct FMediaProfileMenuEntryImpl
 	}
 
 public:
-
+	/** Path of the toolbar that should host the MediaProfile toolbar button. If unspecified by the config, it will default to the Level Editor's toolbar. */
+	FString ToolbarMenu;
 	static TUniquePtr<FMediaProfileMenuEntryImpl> Implementation;
 };
 
