Media Profile Editor: Added ability to mask to specific RGBA channels when displaying media sources and outputs, as well as display an alpha-blended checkerboard background to visualize the alpha channel.

#virtualized
#Synced-CL 48321910

diff --git a/Engine/Plugins/Media/MediaFrameworkUtilities/Source/MediaFrameworkUtilitiesEditor/Private/AssetEditor/MediaProfileEditor.cpp b/Engine/Plugins/Media/MediaFrameworkUtilities/Source/MediaFrameworkUtilitiesEditor/Private/AssetEditor/MediaProfileEditor.cpp
--- a/Engine/Plugins/Media/MediaFrameworkUtilities/Source/MediaFrameworkUtilitiesEditor/Private/AssetEditor/MediaProfileEditor.cpp
+++ b/Engine/Plugins/Media/MediaFrameworkUtilities/Source/MediaFrameworkUtilitiesEditor/Private/AssetEditor/MediaProfileEditor.cpp
@@ -136,6 +136,7 @@ FMediaProfileEditor::~FMediaProfileEditor()
 	FEditorDelegates::PostPIEStarted.RemoveAll(this);
 
 	CloseAllMediaSources();
+	CleanUpAndSaveUserSettings();
 }
 
 void FMediaProfileEditor::RegisterTabSpawners(const TSharedRef<class FTabManager>& InTabManager)
@@ -412,6 +413,51 @@ void FMediaProfileEditor::RefreshSelectedMediaItems(const TArray<int32>& InMedia
 	}
 }
 
+void FMediaProfileEditor::CleanUpAndSaveUserSettings()
+{
+	if (!IsValid(MediaProfileBeingEdited))
+	{
+		return;
+	}
+	
+	UMediaProfileEditorUserSettings* UserSettings = GetMutableDefault<UMediaProfileEditorUserSettings>();
+	if (!UserSettings)
+	{
+		return;
+	}
+	
+	// Remove any per media item settings that reference media items no longer in the media profile
+	TSet<FName> CurrentMediaItemNames;
+	for (int32 Index = 0; Index < MediaProfileBeingEdited->NumMediaSources(); ++Index)
+	{
+		if (UMediaSource* MediaSource = MediaProfileBeingEdited->GetMediaSource(Index))
+		{
+			CurrentMediaItemNames.Add(MediaSource->GetFName());
+		}
+	}
+
+	for (int32 Index = 0; Index < MediaProfileBeingEdited->NumMediaOutputs(); ++Index)
+	{
+		if (UMediaOutput* MediaOutput = MediaProfileBeingEdited->GetMediaOutput(Index))
+		{
+			CurrentMediaItemNames.Add(MediaOutput->GetFName());
+		}
+	}
+	
+	TSet<FName> StoredMediaItemNames;
+	UserSettings->PerMediaItemSettings.GetKeys(StoredMediaItemNames);
+
+	for (const FName& MediaItemName : StoredMediaItemNames)
+	{
+		if (!CurrentMediaItemNames.Contains(MediaItemName))
+		{
+			UserSettings->PerMediaItemSettings.Remove(MediaItemName);
+		}
+	}
+
+	UserSettings->SaveConfig();
+}
+
 void FMediaProfileEditor::OnMediaItemDeleted(UClass* InMediaType, int32 InMediaItemIndex)
 {
 	if (InMediaType->IsChildOf<UMediaSource>())
diff --git a/Engine/Plugins/Media/MediaFrameworkUtilities/Source/MediaFrameworkUtilitiesEditor/Private/AssetEditor/MediaProfileEditor.h b/Engine/Plugins/Media/MediaFrameworkUtilities/Source/MediaFrameworkUtilitiesEditor/Private/AssetEditor/MediaProfileEditor.h
--- a/Engine/Plugins/Media/MediaFrameworkUtilities/Source/MediaFrameworkUtilitiesEditor/Private/AssetEditor/MediaProfileEditor.h
+++ b/Engine/Plugins/Media/MediaFrameworkUtilities/Source/MediaFrameworkUtilitiesEditor/Private/AssetEditor/MediaProfileEditor.h
@@ -81,6 +81,9 @@ private:
 
 	/** Refreshes the selected media items, reopening media sources and outputs */
 	void RefreshSelectedMediaItems(const TArray<int32>& InMediaSources,const TArray<int32>& InMediaOutputs);
+
+	/** Cleans up and saves the media profile editor user settings config */
+	void CleanUpAndSaveUserSettings();
 	
 	void OnMediaItemDeleted(UClass* InMediaType, int32 InMediaItemIndex);
 	void OnSelectedMediaItemsChanged(const TArray<int32>& SelectedMediaSources, const TArray<int32>& SelectedMediaOutputs);
diff --git a/Engine/Plugins/Media/MediaFrameworkUtilities/Source/MediaFrameworkUtilitiesEditor/Private/AssetEditor/MediaProfileEditorUserSettings.h b/Engine/Plugins/Media/MediaFrameworkUtilities/Source/MediaFrameworkUtilitiesEditor/Private/AssetEditor/MediaProfileEditorUserSettings.h
--- a/Engine/Plugins/Media/MediaFrameworkUtilities/Source/MediaFrameworkUtilitiesEditor/Private/AssetEditor/MediaProfileEditorUserSettings.h
+++ b/Engine/Plugins/Media/MediaFrameworkUtilities/Source/MediaFrameworkUtilitiesEditor/Private/AssetEditor/MediaProfileEditorUserSettings.h
@@ -7,6 +7,26 @@
 
 #include "MediaProfileEditorUserSettings.generated.h"
 
+/** Settings for a media source/output in the media profile editor */
+USTRUCT()
+struct FMediaProfileEditorPerMediaItemSettings
+{
+	GENERATED_BODY()
+
+public:
+	/** The saved color channel mask applied to the media source/output */
+	UPROPERTY(Config)
+	int32 ColorChannelMask = INDEX_NONE;
+
+	/** The saved invert alpha flag applied to the media source/output */
+	UPROPERTY(Config)
+	bool bInvertAlphaChannelMask = false;
+
+	/** The saved draw checkerboard flag applied to the media source/output */
+	UPROPERTY(Config)
+	bool bDrawAlphaBlendedCheckerboard = false;
+};
+
 /** Config class to store editor user settings for the media profile editor */
 UCLASS(MinimalAPI, Config=EditorPerProjectUserSettings)
 class UMediaProfileEditorUserSettings : public UObject
@@ -21,6 +41,14 @@ public:
 	/** Whether to show the genlock entry in the viewport toolbar */
 	UPROPERTY(Config)
 	bool bShowGenlockInViewportToolbar = true;
+
+	/** Whether to show the Channels menu in the viewports */
+	UPROPERTY(Config)
+	bool bShowChannelsInViewportToolbar = true;
+
+	/** Config settings specific to individual media sources/outputs, indexed by the media item's name */
+	UPROPERTY(Config)
+	TMap<FName, FMediaProfileEditorPerMediaItemSettings> PerMediaItemSettings;
 };
 
 /**
diff --git a/Engine/Plugins/Media/MediaFrameworkUtilities/Source/MediaFrameworkUtilitiesEditor/Private/AssetEditor/SMaskedMediaImage.cpp b/Engine/Plugins/Media/MediaFrameworkUtilities/Source/MediaFrameworkUtilitiesEditor/Private/AssetEditor/SMaskedMediaImage.cpp
new file mode 100644
--- /dev/null
+++ b/Engine/Plugins/Media/MediaFrameworkUtilities/Source/MediaFrameworkUtilitiesEditor/Private/AssetEditor/SMaskedMediaImage.cpp
@@ -0,0 +1,56 @@
+﻿// Copyright Epic Games, Inc. All Rights Reserved.
+
+#include "SMaskedMediaImage.h"
+
+#include "SceneView.h"
+#include "Engine/Texture.h"
+#include "Materials/Material.h"
+#include "Styling/AppStyle.h"
+#include "Widgets/Images/SImage.h"
+#include "Widgets/Layout/SScaleBox.h"
+
+void SMaskedMediaImage::Construct(const FArguments& InArgs, UTexture* InMediaImageTexture)
+{
+	MediaImageTexture = InMediaImageTexture;
+	ImageSizeAttr = InArgs._ImageSize;
+	ChannelMaskAttr = InArgs._ChannelMask;
+	InvertAlphaChannelAttr = InArgs._InvertAlphaChannel;
+	DrawCheckerboardAttr = InArgs._DrawCheckerboard;
+
+	UMaterial* BaseMaterial = (UMaterial*)StaticLoadObject(UMaterial::StaticClass(), nullptr, TEXT("/MediaFrameworkUtilities/Editor/M_MaskedMediaImage.M_MaskedMediaImage"));
+	check(BaseMaterial);
+	
+	Material = TStrongObjectPtr<UMaterialInstanceDynamic>(UMaterialInstanceDynamic::Create(BaseMaterial, nullptr));
+	Material->SetTextureParameterValue(TEXT("MediaImage"), InMediaImageTexture);
+	Material->SetScalarParameterValue(TEXT("bIsAlphaPremultiplied"), InArgs._IsAlphaPremultiplied ? 1.0f : 0.0f);
+	
+	MaterialBrush = MakeShared<FSlateBrush>();
+	MaterialBrush->SetResourceObject(Material.Get());
+	
+	ChildSlot
+	[
+		SNew(SScaleBox)
+		.Stretch(EStretch::Fill)
+		[
+			SNew(SImage)
+			.Image(MaterialBrush.IsValid() ? MaterialBrush.Get() : FAppStyle::GetBrush("WhiteTexture"))
+		]
+	];
+}
+
+void SMaskedMediaImage::Tick(const FGeometry& AllottedGeometry, const double InCurrentTime, const float InDeltaTime)
+{
+	SCompoundWidget::Tick(AllottedGeometry, InCurrentTime, InDeltaTime);
+
+	if (ImageSizeAttr.IsSet())
+	{
+		FVector2D Size = ImageSizeAttr.Get();
+		MaterialBrush->ImageSize.X = static_cast<float>(Size.X);
+		MaterialBrush->ImageSize.Y = static_cast<float>(Size.Y);
+	}
+
+	const int8 ChannelMask = (int8)ChannelMaskAttr.Get(EColorChannelMask::All);
+	Material->SetScalarParameterValue(TEXT("ChannelMask"), ChannelMask + 1);
+	Material->SetScalarParameterValue(TEXT("bInvertAlphaChannelMask"), InvertAlphaChannelAttr.Get(false) ? 1.0f : 0.0f);
+	Material->SetScalarParameterValue(TEXT("bDrawCheckerboard"), DrawCheckerboardAttr.Get(false) ? 1.0f : 0.0f);
+}
diff --git a/Engine/Plugins/Media/MediaFrameworkUtilities/Source/MediaFrameworkUtilitiesEditor/Private/AssetEditor/SMaskedMediaImage.h b/Engine/Plugins/Media/MediaFrameworkUtilities/Source/MediaFrameworkUtilitiesEditor/Private/AssetEditor/SMaskedMediaImage.h
new file mode 100644
--- /dev/null
+++ b/Engine/Plugins/Media/MediaFrameworkUtilities/Source/MediaFrameworkUtilitiesEditor/Private/AssetEditor/SMaskedMediaImage.h
@@ -0,0 +1,56 @@
+﻿// Copyright Epic Games, Inc. All Rights Reserved.
+
+#pragma once
+
+#include "CoreMinimal.h"
+#include "Widgets/SCompoundWidget.h"
+#include "Materials/MaterialInstanceDynamic.h"
+
+enum class EColorChannelMask : int8;
+class UMaterialInstanceDynamic;
+class UTexture;
+
+/**
+ * Displays a media image texture with optional channel masking and alpha blending support
+ */
+class SMaskedMediaImage : public SCompoundWidget
+{
+public:
+	
+	
+	SLATE_BEGIN_ARGS(SMaskedMediaImage)
+		: _IsAlphaPremultiplied(true)
+	{ }
+		SLATE_ATTRIBUTE(FVector2D, ImageSize)
+		SLATE_ATTRIBUTE(EColorChannelMask, ChannelMask)
+		SLATE_ATTRIBUTE(bool, InvertAlphaChannel)
+		SLATE_ATTRIBUTE(bool, DrawCheckerboard)
+		SLATE_ARGUMENT(bool, IsAlphaPremultiplied)
+	SLATE_END_ARGS()
+
+	void Construct(const FArguments& InArgs, UTexture* InMediaImageTexture);
+
+	virtual void Tick(const FGeometry& AllottedGeometry, const double InCurrentTime, const float InDeltaTime) override;
+	
+private:
+	/** Material used to render the media image texture */
+	TStrongObjectPtr<UMaterialInstanceDynamic> Material;
+
+	/** The texture the media image is being written to */
+	TWeakObjectPtr<UTexture> MediaImageTexture;
+	
+	/** Slate brush that wraps the material */
+	TSharedPtr<FSlateBrush> MaterialBrush;
+	
+	/** Attribute used to size the media image appropriately */
+	TAttribute<FVector2D> ImageSizeAttr;
+
+	/** Attribute for the color channel to mask to */
+	TAttribute<EColorChannelMask> ChannelMaskAttr;
+
+	/** Attribute for the invert alpha flag */
+	TAttribute<bool> InvertAlphaChannelAttr;
+
+	/** Attribute for drawing the alpha-blended checkerboard background */
+	TAttribute<bool> DrawCheckerboardAttr;
+};
diff --git a/Engine/Plugins/Media/MediaFrameworkUtilities/Source/MediaFrameworkUtilitiesEditor/Private/AssetEditor/SMediaProfileMediaItemDisplay.cpp b/Engine/Plugins/Media/MediaFrameworkUtilities/Source/MediaFrameworkUtilitiesEditor/Private/AssetEditor/SMediaProfileMediaItemDisplay.cpp
--- a/Engine/Plugins/Media/MediaFrameworkUtilities/Source/MediaFrameworkUtilitiesEditor/Private/AssetEditor/SMediaProfileMediaItemDisplay.cpp
+++ b/Engine/Plugins/Media/MediaFrameworkUtilities/Source/MediaFrameworkUtilitiesEditor/Private/AssetEditor/SMediaProfileMediaItemDisplay.cpp
@@ -8,13 +8,13 @@
 #include "MediaPlayer.h"
 #include "MediaTexture.h"
 #include "ShowFlagMenuCommands.h"
+#include "SMaskedMediaImage.h"
 #include "CaptureTab/SMediaFrameworkCapture.h"
 #include "Engine/TextureRenderTarget2D.h"
 #include "Profile/MediaProfilePlaybackManager.h"
 #include "Slate/SceneViewport.h"
 #include "UI/MediaFrameworkUtilitiesEditorStyle.h"
 #include "ViewportToolbar/UnrealEdViewportToolbar.h"
-#include "Widgets/SMediaImage.h"
 #include "Widgets/SViewport.h"
 #include "Widgets/Images/SImage.h"
 #include "Widgets/Text/STextBlock.h"
@@ -104,8 +104,12 @@ void SMediaProfileMediaSourceDisplay::ConfigureMediaImage()
 		.MinAspectRatio(this, &SMediaProfileMediaSourceDisplay::GetSourceAspectRatio)
 		.MaxAspectRatio(this, &SMediaProfileMediaSourceDisplay::GetSourceAspectRatio)
 		[
-			SNew(SMediaImage, MediaTexture)
-			.BrushImageSize(TextureSize)
+			SNew(SMaskedMediaImage, MediaTexture)
+			.IsAlphaPremultiplied(false)
+			.ChannelMask(this, &SMediaProfileMediaSourceDisplay::GetChannelMask)
+			.InvertAlphaChannel(this, &SMediaProfileMediaSourceDisplay::GetInvertAlphaChannelMask)
+			.DrawCheckerboard(this, &SMediaProfileMediaSourceDisplay::GetDrawAlphaBlendedCheckerboard)
+			.ImageSize(TextureSize)
 		]);
 }
 
@@ -257,6 +261,15 @@ public:
 		];
 	}
 	
+	/** Sets the current color channel being displayed in the viewport */
+	virtual void SetChannelMask(EColorChannelMask InChannelMask) { }
+
+	/** Sets whether the alpha color channel is inverted when displayed as the color channel mask */
+	virtual void SetAlphaInverted(bool bInAlphaInverted) { }
+
+	/** Sets whether to draw the alpha blended checkerboard background when rendering the viewport */
+	virtual void SetDrawAlphaBlendedCheckerboard(bool bInDrawAlphaBlendedCheckerboard) { }
+
 protected:
 	virtual TSharedRef<SWidget> GetContentWidget() { return SNullWidget::NullWidget; }
 	
@@ -402,6 +415,65 @@ public:
 		SMediaCaptureImageBase::Construct(BaseArgs);
 	}
 
+	virtual void SetChannelMask(EColorChannelMask InChannelMask) override
+	{
+		TWeakPtr<FSceneViewport> Viewport = GetActiveEditorViewport();
+		if (TSharedPtr<FSceneViewport> PinnedViewport = Viewport.Pin())
+		{
+			if (FEditorViewportClient* EditorViewportClient = static_cast<FEditorViewportClient*>(PinnedViewport->GetClient()))
+			{
+				EditorViewportClient->ChannelMaskParams.ColorChannelMask = InChannelMask;
+				EditorViewportClient->Invalidate();
+			}
+		}
+	}
+	
+	virtual void SetAlphaInverted(bool bInAlphaInverted) override
+	{
+		TWeakPtr<FSceneViewport> Viewport = GetActiveEditorViewport();
+		if (TSharedPtr<FSceneViewport> PinnedViewport = Viewport.Pin())
+		{
+			if (FEditorViewportClient* EditorViewportClient = static_cast<FEditorViewportClient*>(PinnedViewport->GetClient()))
+			{
+				EditorViewportClient->ChannelMaskParams.bInvertAlphaChannelMask = bInAlphaInverted;
+				
+				if (!bInAlphaInverted)
+				{
+					EditorViewportClient->ChannelMaskParams.bDrawAlphaBlendedCheckerboard = false;
+				}
+				
+				EditorViewportClient->Invalidate();
+			}
+		}
+	}
+
+	virtual void SetDrawAlphaBlendedCheckerboard(bool bInDrawAlphaBlendedCheckerboard) override
+	{
+		TWeakPtr<FSceneViewport> Viewport = GetActiveEditorViewport();
+		if (TSharedPtr<FSceneViewport> PinnedViewport = Viewport.Pin())
+		{
+			if (FEditorViewportClient* EditorViewportClient = static_cast<FEditorViewportClient*>(PinnedViewport->GetClient()))
+			{
+				EditorViewportClient->ChannelMaskParams.bDrawAlphaBlendedCheckerboard = bInDrawAlphaBlendedCheckerboard;
+				EditorViewportClient->Invalidate();
+			}
+		}
+	}
+
+	void SetMediaItemConfigSettings(FMediaProfileEditorPerMediaItemSettings* InMediaItemSettings)
+	{
+		TWeakPtr<FSceneViewport> Viewport = GetActiveEditorViewport();
+		if (TSharedPtr<FSceneViewport> PinnedViewport = Viewport.Pin())
+		{
+			if (FEditorViewportClient* EditorViewportClient = static_cast<FEditorViewportClient*>(PinnedViewport->GetClient()))
+			{
+				InMediaItemSettings->ColorChannelMask = (int32)EditorViewportClient->ChannelMaskParams.ColorChannelMask;
+				InMediaItemSettings->bInvertAlphaChannelMask = EditorViewportClient->ChannelMaskParams.bInvertAlphaChannelMask;
+				InMediaItemSettings->bDrawAlphaBlendedCheckerboard = EditorViewportClient->ChannelMaskParams.bDrawAlphaBlendedCheckerboard;
+			}
+		}
+	}
+	
 protected:
 	virtual TSharedRef<SWidget> GetContentWidget() override
 	{
@@ -504,6 +576,45 @@ public:
 		}
 	}
 	
+	virtual void SetChannelMask(EColorChannelMask InChannelMask) override
+	{
+		if (!ViewportClient.IsValid())
+		{
+			return;
+		}
+		
+		ViewportClient->ChannelMaskParams.ColorChannelMask = InChannelMask;
+		ViewportClient->Invalidate();
+	}
+
+	virtual void SetAlphaInverted(bool bInAlphaInverted) override
+	{
+		if (!ViewportClient.IsValid())
+		{
+			return;
+		}
+		
+		ViewportClient->ChannelMaskParams.bInvertAlphaChannelMask = bInAlphaInverted;
+
+		if (!bInAlphaInverted)
+		{
+			ViewportClient->ChannelMaskParams.bDrawAlphaBlendedCheckerboard = false;
+		}
+		
+		ViewportClient->Invalidate();
+	}
+
+	virtual void SetDrawAlphaBlendedCheckerboard(bool bInDrawAlphaBlendedCheckerboard) override
+	{
+		if (!ViewportClient.IsValid())
+		{
+			return;
+		}
+		
+		ViewportClient->ChannelMaskParams.bDrawAlphaBlendedCheckerboard = bInDrawAlphaBlendedCheckerboard;
+		ViewportClient->Invalidate();
+	}
+
 	/** Updates the display if the capture's locked actors have been changed */
 	void Refresh(const TArray<TSoftObjectPtr<AActor>>& InLockedActors)
 	{
@@ -648,10 +759,17 @@ public:
 	SLATE_BEGIN_ARGS(SMediaCaptureRenderTargetImage) {}
 		SLATE_ATTRIBUTE(TOptional<EMediaCaptureState>, MediaCaptureState)
 		SLATE_ATTRIBUTE(FOptionalSize, AspectRatio)
+		SLATE_ATTRIBUTE(EColorChannelMask, ChannelMask)
+		SLATE_ATTRIBUTE(bool, InvertAlphaChannel)
+		SLATE_ATTRIBUTE(bool, DrawCheckerboard)
 	SLATE_END_ARGS()
 
 	void Construct(const FArguments& InArgs, UTextureRenderTarget2D* InRenderTarget)
 	{
+		ChannelMaskAttr = InArgs._ChannelMask;
+		InvertAlphaChannelAttr = InArgs._InvertAlphaChannel;
+		DrawCheckerboardAttr = InArgs._DrawCheckerboard;
+		
 		SMediaCaptureImageBase::FArguments BaseArgs;
 		BaseArgs.MediaCaptureState(InArgs._MediaCaptureState);
 		BaseArgs.AspectRatio(InArgs._AspectRatio);
@@ -674,7 +792,11 @@ public:
 		{
 			Container->SetContent
 			(
-				SNew(SMediaImage, InRenderTarget)
+				SNew(SMaskedMediaImage, InRenderTarget)
+				.IsAlphaPremultiplied(true)
+				.ChannelMask(ChannelMaskAttr)
+				.InvertAlphaChannel(InvertAlphaChannelAttr)
+				.DrawCheckerboard(DrawCheckerboardAttr)
 			);
 		}
 		else
@@ -706,8 +828,14 @@ private:
 	/** The render target being used by the capture */
 	TWeakObjectPtr<UTextureRenderTarget2D> RenderTarget;
 	
-	/** Attribute for querying if the viewport capture is currently actively being captured */
-	TAttribute<bool> IsCapturing;
+	/** Attribute for the color channel to mask to */
+	TAttribute<EColorChannelMask> ChannelMaskAttr;
+
+	/** Attribute for the invert alpha flag */
+	TAttribute<bool> InvertAlphaChannelAttr;
+
+	/** Attribute for drawing the alpha-blended checkerboard background */
+	TAttribute<bool> DrawCheckerboardAttr;
 };
 
 void SMediaProfileMediaOutputDisplay::Construct(const FArguments& InArgs, const TSharedPtr<SMediaProfileViewport>& InOwningViewport)
@@ -772,11 +900,22 @@ void SMediaProfileMediaOutputDisplay::ConfigureMediaImage()
 	{
 		if (CaptureSettings->CurrentViewportMediaOutput.MediaOutput == MediaOutput)
 		{
+			TSharedPtr<SMediaCaptureEditorViewportImage> CurrentViewportImage;
 			MediaImageContainer->SetContent(
-				SAssignNew(ImageWidget, SMediaCaptureEditorViewportImage, MediaProfile, CaptureSettings->CurrentViewportMediaOutput)
+				SAssignNew(CurrentViewportImage, SMediaCaptureEditorViewportImage, MediaProfile, CaptureSettings->CurrentViewportMediaOutput)
 				.MediaCaptureState(this, &SMediaProfileMediaOutputDisplay::GetMediaOutputCaptureState)
 				.AspectRatio(this, &SMediaProfileMediaOutputDisplay::GetMediaOutputDesiredAspectRatio)
 			);
+
+			ImageWidget = CurrentViewportImage;
+
+			// Instead of overriding with the color channel masking settings from the config, respect the current viewport's settings,
+			// and write them to the config instead
+			if (FMediaProfileEditorPerMediaItemSettings* MediaItemSettings = GetMediaItemSettings())
+			{
+				CurrentViewportImage->SetMediaItemConfigSettings(MediaItemSettings);
+			}
+			
 			return;
 		}
 		
@@ -806,6 +945,15 @@ void SMediaProfileMediaOutputDisplay::ConfigureMediaImage()
 			
 			FShowFlagMenuCommands::Get().BindCommands(*ShowFlagsCommandList.Get(), ViewportImage->GetViewportClient());
 			bDisplayShowFlags = true;
+
+			// Sync the viewport's color channel mask settings with the saved settings in the editor config
+			if (FMediaProfileEditorPerMediaItemSettings* MediaItemSettings = GetMediaItemSettings())
+			{
+				ViewportImage->SetChannelMask((EColorChannelMask)MediaItemSettings->ColorChannelMask);
+				ViewportImage->SetAlphaInverted(MediaItemSettings->bInvertAlphaChannelMask);
+				ViewportImage->SetDrawAlphaBlendedCheckerboard(MediaItemSettings->bDrawAlphaBlendedCheckerboard);
+			}
+			
 			return;
 		}
 
@@ -821,6 +969,9 @@ void SMediaProfileMediaOutputDisplay::ConfigureMediaImage()
 				SAssignNew(ImageWidget, SMediaCaptureRenderTargetImage, CaptureSettings->RenderTargetCaptures[RenderTargetOutputInfoIndex].RenderTarget)
 				.MediaCaptureState(this, &SMediaProfileMediaOutputDisplay::GetMediaOutputCaptureState)
 				.AspectRatio(this, &SMediaProfileMediaOutputDisplay::GetMediaOutputDesiredAspectRatio)
+				.ChannelMask(this, &SMediaProfileMediaOutputDisplay::GetChannelMask)
+				.InvertAlphaChannel(this, &SMediaProfileMediaOutputDisplay::GetInvertAlphaChannelMask)
+				.DrawCheckerboard(this, &SMediaProfileMediaOutputDisplay::GetDrawAlphaBlendedCheckerboard)
 			);
 			return;
 		}
@@ -863,6 +1014,51 @@ FText SMediaProfileMediaOutputDisplay::GetBaseMediaTypeLabel() const
 	return LOCTEXT("MediaOutputBaseTypeLabel", "Media Output");
 }
 
+void SMediaProfileMediaOutputDisplay::SetChannelMask(EColorChannelMask InChannelMask)
+{
+	SMediaProfileMediaItemDisplayBase::SetChannelMask(InChannelMask);
+	
+	if (TSharedPtr<SMediaCaptureImageBase> CaptureImageWidget = StaticCastSharedPtr<SMediaCaptureImageBase>(ImageWidget))
+    {
+    	CaptureImageWidget->SetChannelMask(InChannelMask);
+    }
+}
+
+void SMediaProfileMediaOutputDisplay::SetInvertAlphaChannelMask(bool bInInvertAlphaChannelMask)
+{
+	SMediaProfileMediaItemDisplayBase::SetInvertAlphaChannelMask(bInInvertAlphaChannelMask);
+	
+	if (TSharedPtr<SMediaCaptureImageBase> CaptureImageWidget = StaticCastSharedPtr<SMediaCaptureImageBase>(ImageWidget))
+	{
+		CaptureImageWidget->SetAlphaInverted(bInInvertAlphaChannelMask);
+
+		if (!bInInvertAlphaChannelMask && RenderTargetOutputInfoIndex == INDEX_NONE)
+		{
+			// For viewports, automatically disable the checkerboard if the alpha is set to non-inverted match the behavior
+			// of the level editor viewport RGBA feature.
+			SetDrawAlphaBlendedCheckerboard(false);
+		}
+	}
+}
+
+void SMediaProfileMediaOutputDisplay::SetDrawAlphaBlendedCheckerboard(bool bInDrawAlphaBlendedCheckerboard)
+{
+	SMediaProfileMediaItemDisplayBase::SetDrawAlphaBlendedCheckerboard(bInDrawAlphaBlendedCheckerboard);
+	
+	if (TSharedPtr<SMediaCaptureImageBase> CaptureImageWidget = StaticCastSharedPtr<SMediaCaptureImageBase>(ImageWidget))
+	{
+		return CaptureImageWidget->SetDrawAlphaBlendedCheckerboard(bInDrawAlphaBlendedCheckerboard);
+	}
+}
+
+bool SMediaProfileMediaOutputDisplay::IsDrawAlphaBlendedCheckerboardEnabled() const
+{
+	// For viewports, don't allow checkerboard to be enabled if the alpha is not inverted to match the behavior
+	// of the level editor viewport RGBA feature. Allow it for render targets since render targets may have
+	// a standard alpha content instead of scene color
+	return RenderTargetOutputInfoIndex != INDEX_NONE ? true : GetInvertAlphaChannelMask();
+}
+
 void SMediaProfileMediaOutputDisplay::AddToolbarEntries(FToolMenuSection& Section)
 {
 	FToolMenuEntry& ShowSubmenuEntry = Section.AddEntry(UE::UnrealEd::CreateShowSubmenu(
diff --git a/Engine/Plugins/Media/MediaFrameworkUtilities/Source/MediaFrameworkUtilitiesEditor/Private/AssetEditor/SMediaProfileMediaItemDisplay.h b/Engine/Plugins/Media/MediaFrameworkUtilities/Source/MediaFrameworkUtilitiesEditor/Private/AssetEditor/SMediaProfileMediaItemDisplay.h
--- a/Engine/Plugins/Media/MediaFrameworkUtilities/Source/MediaFrameworkUtilitiesEditor/Private/AssetEditor/SMediaProfileMediaItemDisplay.h
+++ b/Engine/Plugins/Media/MediaFrameworkUtilities/Source/MediaFrameworkUtilitiesEditor/Private/AssetEditor/SMediaProfileMediaItemDisplay.h
@@ -2,12 +2,14 @@
 
 #pragma once
 
+#include "EditorViewportCommands.h"
 #include "LevelViewportActions.h"
 #include "MediaOutput.h"
 #include "MediaProfileEditor.h"
 #include "MediaProfileEditorUserSettings.h"
 #include "MediaSource.h"
 #include "PropertyEditorModule.h"
+#include "SceneView.h"
 #include "SMediaProfileViewport.h"
 #include "ToolMenus.h"
 #include "ToolMenusEditor.h"
@@ -226,6 +228,41 @@ protected:
 					return IsMaximized() ? PinnedOwningViewport->CanRestorePreviousLayout() : PinnedOwningViewport->CanMaximizePanel(PanelIndex);
 				})
 			));
+
+		CommandList->MapAction(FEditorViewportCommands::Get().AllChannelsMask,
+			FUIAction(
+				FExecuteAction::CreateSP(this, &SMediaProfileMediaItemDisplayBase::SetChannelMask, EColorChannelMask::All),
+				FCanExecuteAction(),
+				FIsActionChecked::CreateSP(this, &SMediaProfileMediaItemDisplayBase::IsChannelBeingMasked, EColorChannelMask::All)
+			));
+		
+		CommandList->MapAction(FEditorViewportCommands::Get().RedChannelMask,
+			FUIAction(
+				FExecuteAction::CreateSP(this, &SMediaProfileMediaItemDisplayBase::SetChannelMask, EColorChannelMask::Red),
+				FCanExecuteAction(),
+				FIsActionChecked::CreateSP(this, &SMediaProfileMediaItemDisplayBase::IsChannelBeingMasked, EColorChannelMask::Red)
+			));
+
+		CommandList->MapAction(FEditorViewportCommands::Get().GreenChannelMask,
+			FUIAction(
+				FExecuteAction::CreateSP(this, &SMediaProfileMediaItemDisplayBase::SetChannelMask, EColorChannelMask::Green),
+				FCanExecuteAction(),
+				FIsActionChecked::CreateSP(this, &SMediaProfileMediaItemDisplayBase::IsChannelBeingMasked, EColorChannelMask::Green)
+			));
+
+		CommandList->MapAction(FEditorViewportCommands::Get().BlueChannelMask,
+			FUIAction(
+				FExecuteAction::CreateSP(this, &SMediaProfileMediaItemDisplayBase::SetChannelMask, EColorChannelMask::Blue),
+				FCanExecuteAction(),
+				FIsActionChecked::CreateSP(this, &SMediaProfileMediaItemDisplayBase::IsChannelBeingMasked, EColorChannelMask::Blue)
+			));
+
+		CommandList->MapAction(FEditorViewportCommands::Get().AlphaChannelMask,
+			FUIAction(
+				FExecuteAction::CreateSP(this, &SMediaProfileMediaItemDisplayBase::SetChannelMask, EColorChannelMask::Alpha),
+				FCanExecuteAction(),
+				FIsActionChecked::CreateSP(this, &SMediaProfileMediaItemDisplayBase::IsChannelBeingMasked, EColorChannelMask::Alpha)
+			));
 	}
 
 	/** Creates the viewport display's toolbar */
@@ -293,6 +330,7 @@ protected:
 				}
 				
 				Section.AddEntry(ContextWidget->CreateMediaItemsComboButton());
+				Section.AddEntry(ContextWidget->CreateChannelsButton());
 				ContextWidget->AddToolbarEntries(Section);
 				Section.AddEntry(ContextWidget->CreateViewportArrangementMenu());
 
@@ -389,6 +427,153 @@ protected:
 			
 		LockButton.SetShowInToolbarTopLevel(true);
 		LockButton.ToolBarData.ResizeParams.AllowClipping = false;
+
+		FToolMenuSection& ChannelsSection = ToolMenu->AddSection(TEXT("ChannelsSection"), LOCTEXT("ChannelsSectionLabel", "Channels"));
+		ChannelsSection.AddEntry(CreateChannelsSubmenu(true));
+	}
+
+	FToolMenuEntry CreateChannelsSubmenu(bool bUseFullLabel)
+	{
+		static FText ChannelLabels[] =
+		{
+			LOCTEXT("ChannelRGB_Label", "RGB"),
+			LOCTEXT("ChannelRed_Label", "Red"),
+			LOCTEXT("ChannelGreen_Label", "Green"),
+			LOCTEXT("ChannelBlue_Label", "Blue"),
+			LOCTEXT("ChannelAlpha_Label", "Alpha")
+		};
+
+		static FName ChannelIcons[] =
+		{
+			"EditorViewport.AllChannelsMask",
+			"EditorViewport.RedChannelMask",
+			"EditorViewport.GreenChannelMask",
+			"EditorViewport.BlueChannelMask",
+			"EditorViewport.AlphaChannelMask"
+		};
+
+		FToolMenuEntry Entry = FToolMenuEntry::InitSubMenu(
+			"Channels",
+			TAttribute<FText>::CreateLambda([this, bUseFullLabel]
+			{
+				const int8 ChannelMaskIndex = (int8)GetChannelMask();
+				return bUseFullLabel
+					? FText::Format(LOCTEXT("ChannelsMenu_LabelFormat", "Channels ({0})"), ChannelLabels[ChannelMaskIndex + 1])
+					: ChannelLabels[ChannelMaskIndex + 1];
+			}),
+			LOCTEXT("ChannelsMenu_ToolTip", "Display specific color channels in this media source/output"),
+			FNewToolMenuDelegate::CreateLambda([this](UToolMenu* InMenu)
+			{
+				FToolMenuSection& ChannelsSubMenuSection =
+					InMenu->AddSection("ChannelsSubMenuSection", LOCTEXT("ChannelsSubMenuSectionHeader", "Channels"));
+
+				ChannelsSubMenuSection.AddMenuEntry(FEditorViewportCommands::Get().AllChannelsMask);
+				ChannelsSubMenuSection.AddMenuEntry(FEditorViewportCommands::Get().RedChannelMask);
+				ChannelsSubMenuSection.AddMenuEntry(FEditorViewportCommands::Get().GreenChannelMask);
+				ChannelsSubMenuSection.AddMenuEntry(FEditorViewportCommands::Get().BlueChannelMask);
+				ChannelsSubMenuSection.AddMenuEntry(FEditorViewportCommands::Get().AlphaChannelMask);
+
+				FToolMenuSection& OptionsSubMenuSection =
+					InMenu->AddSection("OptionsSubMenuSection", LOCTEXT("OptionsSubMenuSectionHeader", "Options"));
+
+				OptionsSubMenuSection.AddMenuEntry(
+					"ShowChannelsInViewportToolbar",
+					LOCTEXT("ShowChannelsInToolbarLabel", "Show Channels in Viewport Toolbar"),
+					LOCTEXT("ShowChannelsInToolbarToolTip", "Show the channel selector in the viewport toolbar"),
+					FSlateIcon(),
+					FUIAction(
+						FExecuteAction::CreateLambda([]()
+						{
+							if (UMediaProfileEditorUserSettings* UserSettings = GetMutableDefault<UMediaProfileEditorUserSettings>())
+							{
+								UserSettings->bShowChannelsInViewportToolbar = !UserSettings->bShowChannelsInViewportToolbar;
+							}
+						}),
+						FCanExecuteAction(),
+						FGetActionCheckState::CreateLambda([]()
+						{
+							return GetDefault<UMediaProfileEditorUserSettings>()->bShowChannelsInViewportToolbar ? ECheckBoxState::Checked : ECheckBoxState::Unchecked;
+						})
+					),
+					EUserInterfaceActionType::ToggleButton
+				);
+
+				FToolMenuEntry& InvertAlphaEntry = OptionsSubMenuSection.AddMenuEntry(
+					"InvertAlpha",
+					LOCTEXT("InvertAlphaLabel", "View Alpha as Inverted"),
+					LOCTEXT("InvertAlphaToolTip", "Invert the alpha value displayed when viewing the alpha channel"),
+					FSlateIcon(),
+					FUIAction(
+						FExecuteAction::CreateLambda([this]
+						{
+							SetInvertAlphaChannelMask(!GetInvertAlphaChannelMask());
+						}),
+						FCanExecuteAction(),
+						FGetActionCheckState::CreateLambda([this]
+						{
+							return GetInvertAlphaChannelMask() ? ECheckBoxState::Checked : ECheckBoxState::Unchecked;
+						})
+					),
+					EUserInterfaceActionType::ToggleButton
+				);
+				InvertAlphaEntry.Visibility = TAttribute<bool>::CreateSP(this, &SMediaProfileMediaItemDisplayBase::CanInvertAlphaChannelMask);
+
+				OptionsSubMenuSection.AddMenuEntry(
+					"VisualizeAlphaWithCheckerboard",
+					LOCTEXT("VisualizeWithCheckerboardLabel", "Visualize Alpha with Checkerboard"),
+					LOCTEXT("VisualizeWithCheckerboardToolTip", "Displays a checkerboard underneath the viewport to help visualize the alpha channel"),
+					FSlateIcon(),
+					FUIAction(
+						FExecuteAction::CreateLambda([this]
+						{
+							SetDrawAlphaBlendedCheckerboard(!GetDrawAlphaBlendedCheckerboard());
+						}),
+						FCanExecuteAction::CreateLambda([this]
+						{
+							return IsDrawAlphaBlendedCheckerboardEnabled();
+						}),
+						FGetActionCheckState::CreateLambda([this]
+						{
+							return GetDrawAlphaBlendedCheckerboard() ? ECheckBoxState::Checked : ECheckBoxState::Unchecked;
+						})
+					),
+					EUserInterfaceActionType::ToggleButton
+				);
+			}),
+			/*bInOpenSubMenuOnClick=*/false,
+			TAttribute<FSlateIcon>::CreateLambda(
+				[this]
+				{
+					const int8 ChannelMaskIndex = (int8)GetChannelMask();
+					return FSlateIcon(FAppStyle::GetAppStyleSetName(), ChannelIcons[ChannelMaskIndex + 1]);
+				})
+		);
+
+		return Entry;
+	}
+
+	FToolMenuEntry CreateChannelsButton()
+	{
+		return FToolMenuEntry::InitDynamicEntry(
+			"Channels",
+			FNewToolMenuSectionDelegate::CreateLambda(
+				[](FToolMenuSection& InDynamicSection) -> void
+				{
+					if (UMediaProfileMediaItemDisplayContext* Context = InDynamicSection.FindContext<UMediaProfileMediaItemDisplayContext>())
+					{
+						TSharedPtr<SMediaProfileMediaItemDisplayBase<TMediaItem>> ContextWidget = Context->GetDisplayWidget<SMediaProfileMediaItemDisplayBase<TMediaItem>>();
+						
+						constexpr bool bUseFullLabel = false;
+						FToolMenuEntry& Entry = InDynamicSection.AddEntry(ContextWidget->CreateChannelsSubmenu(false));
+						Entry.Visibility = TAttribute<bool>::CreateLambda([]
+						{
+							return GetDefault<UMediaProfileEditorUserSettings>()->bShowChannelsInViewportToolbar;
+						});
+						Entry.ToolBarData.ResizeParams.ClippingPriority = 800;
+					}
+				}
+			)
+		);
 	}
 
 	/** Creates the viewport arrangement menu, which allows users to change the arrangement of panels in the viewport */
@@ -605,6 +790,88 @@ protected:
 		ChangePanelContents(InMediaItemClass, InMediaItemIndex, bRefreshDisplay);
 	}
 
+	/** Gets the current color channel being displayed in the viewport */
+	virtual EColorChannelMask GetChannelMask() const
+	{
+		if (FMediaProfileEditorPerMediaItemSettings* MediaItemSettings = GetMediaItemSettings())
+		{
+			return (EColorChannelMask)MediaItemSettings->ColorChannelMask;
+		}
+
+		return EColorChannelMask::All;
+	}
+
+	/** Sets the current color channel being displayed in the viewport */
+	virtual void SetChannelMask(EColorChannelMask InChannelMask)
+	{
+		if (FMediaProfileEditorPerMediaItemSettings* MediaItemSettings = GetMediaItemSettings())
+		{
+			MediaItemSettings->ColorChannelMask = (int32)InChannelMask;
+		}
+	}
+
+	/** Gets whether the specified color channel is the one being masked to */
+	virtual bool IsChannelBeingMasked(EColorChannelMask InChannelMask)
+	{
+		if (FMediaProfileEditorPerMediaItemSettings* MediaItemSettings = GetMediaItemSettings())
+		{
+			return (EColorChannelMask)MediaItemSettings->ColorChannelMask == InChannelMask;
+		}
+
+		return false;
+	}
+	
+	/** Gets whether the alpha color channel is inverted when displayed as the color channel mask */
+	virtual bool GetInvertAlphaChannelMask() const
+	{
+		if (FMediaProfileEditorPerMediaItemSettings* MediaItemSettings = GetMediaItemSettings())
+		{
+			return MediaItemSettings->bInvertAlphaChannelMask;
+		}
+
+		return false;
+	}
+
+	/** Gets whether the alpha color channel is inverted when displayed as the color channel mask */
+	virtual void SetInvertAlphaChannelMask(bool bInInvertAlphaChannelMask)
+	{
+		if (FMediaProfileEditorPerMediaItemSettings* MediaItemSettings = GetMediaItemSettings())
+		{
+			MediaItemSettings->bInvertAlphaChannelMask = bInInvertAlphaChannelMask;
+		}
+	}
+
+	/** Gets whether the display's alpha channel can be displayed as inverted or not */
+	virtual bool CanInvertAlphaChannelMask() const
+	{
+		return false;
+	}
+	
+	/** Gets whether to draw the alpha blended checkerboard background when rendering the viewport */
+	virtual bool GetDrawAlphaBlendedCheckerboard() const
+	{
+		if (FMediaProfileEditorPerMediaItemSettings* MediaItemSettings = GetMediaItemSettings())
+		{
+			return MediaItemSettings->bDrawAlphaBlendedCheckerboard;
+		}
+
+		return false;
+	}
+
+	/** Sets whether to draw the alpha blended checkerboard background when rendering the viewport */
+	virtual void SetDrawAlphaBlendedCheckerboard(bool bInDrawAlphaBlendedCheckerboard)
+	{
+		if (FMediaProfileEditorPerMediaItemSettings* MediaItemSettings = GetMediaItemSettings())
+		{
+			MediaItemSettings->bDrawAlphaBlendedCheckerboard = bInDrawAlphaBlendedCheckerboard;
+		}
+	}
+
+	virtual bool IsDrawAlphaBlendedCheckerboardEnabled() const
+	{
+		return true;
+	}
+
 	/** Gets whether this panel is currently maximized in the viewport */
 	bool IsMaximized() const
 	{
@@ -667,6 +934,25 @@ protected:
 			FSlateIcon(FAppStyle::GetAppStyleSetName(), TEXT("PropertyWindow.Unlocked"));
 	}
 
+	FMediaProfileEditorPerMediaItemSettings* GetMediaItemSettings() const
+	{
+		if (UMediaProfileEditorUserSettings* UserSettings = GetMutableDefault<UMediaProfileEditorUserSettings>())
+		{
+			if (TMediaItem* MediaItem = GetMediaItem())
+			{
+				FName MediaItemName = MediaItem->GetFName();
+				if (!UserSettings->PerMediaItemSettings.Contains(MediaItemName))
+				{
+					UserSettings->PerMediaItemSettings.Add(MediaItemName, FMediaProfileEditorPerMediaItemSettings());
+				}
+
+				return &UserSettings->PerMediaItemSettings[MediaItemName];
+			}
+		}
+
+		return nullptr;
+	}
+	
 	/** Raised when an object property is changed in the editor, used to potentially refresh the display if the displayed media item is modified */
 	virtual void OnObjectPropertyChanged(UObject* InObject, FPropertyChangedEvent& InPropertyChangedEvent)
 	{
@@ -753,6 +1039,12 @@ protected:
 	virtual FText GetMediaItemLabel() const override;
 	virtual FText GetBaseMediaTypeLabel() const override;
 
+	virtual void SetChannelMask(EColorChannelMask InChannelMask) override;
+	virtual void SetInvertAlphaChannelMask(bool bInInvertAlphaChannelMask) override;
+	virtual bool CanInvertAlphaChannelMask() const override { return true; }
+	virtual void SetDrawAlphaBlendedCheckerboard(bool bInDrawAlphaBlendedCheckerboard) override;
+	virtual bool IsDrawAlphaBlendedCheckerboardEnabled() const override;
+
 	virtual void AddToolbarEntries(FToolMenuSection& Section) override;
 	virtual void AppendToolbarCommandList(FToolMenuContext& Context) override;
 	
