Media Profile: Disable alpha-blended checkerboard visualization for media outputs as the checkerboard can inadvertently be passed to the media output stream due to how it is implemented for scene viewports.

#Synced-CL 49224024

diff --git a/Engine/Plugins/Media/MediaFrameworkUtilities/Source/MediaFrameworkUtilities/Private/Profile/MediaProfilePlaybackManager.cpp b/Engine/Plugins/Media/MediaFrameworkUtilities/Source/MediaFrameworkUtilities/Private/Profile/MediaProfilePlaybackManager.cpp
--- a/Engine/Plugins/Media/MediaFrameworkUtilities/Source/MediaFrameworkUtilities/Private/Profile/MediaProfilePlaybackManager.cpp
+++ b/Engine/Plugins/Media/MediaFrameworkUtilities/Source/MediaFrameworkUtilities/Private/Profile/MediaProfilePlaybackManager.cpp
@@ -237,14 +237,19 @@ namespace MediaProfilePlaybackManager
 					CachedViewportFlags.bDrawAxes = ViewportClient.bDrawAxes;
 					CachedViewportFlags.bDisableInput = ViewportClient.bDisableInput;
 					CachedViewportFlags.bAllowCinematicControl = ViewportClient.AllowsCinematicControl();
+					CachedViewportFlags.bDrawAlphaBlendedCheckerboard = ViewportClient.ChannelMaskParams.bDrawAlphaBlendedCheckerboard;
 					CachedViewportFlags.VisibilityDelegate = ViewportClient.VisibilityDelegate;
 
 					ViewportClient.SetRealtime(true);
 					ViewportClient.bSetListenerPosition = false;
 					ViewportClient.bDrawAxes = false;
 					ViewportClient.SetAllowCinematicControl(false);
+					ViewportClient.ChannelMaskParams.bDrawAlphaBlendedCheckerboard = false;
 					ViewportClient.VisibilityDelegate.BindLambda([] { return true; });
 
+					FLevelEditorModule& LevelEditorModule = FModuleManager::GetModuleChecked<FLevelEditorModule>(TEXT("LevelEditor"));
+					LevelEditorModule.OnGetVisualizeCheckerboardEnabled().BindLambda([] { return false; });
+
 					const bool bDisableInput = InCaptureOptions.ResizeMethod != EMediaCaptureResizeMethod::ResizeInRenderPass;
 					ViewportClient.bDisableInput = bDisableInput ? true : ViewportClient.bDisableInput;
 
@@ -284,8 +289,12 @@ namespace MediaProfilePlaybackManager
 				ViewportClient.bDrawAxes = CachedViewportFlags.bDrawAxes;
 				ViewportClient.bDisableInput = CachedViewportFlags.bDisableInput;
 				ViewportClient.SetAllowCinematicControl(CachedViewportFlags.bAllowCinematicControl);
+				ViewportClient.ChannelMaskParams.bDrawAlphaBlendedCheckerboard = CachedViewportFlags.bDrawAlphaBlendedCheckerboard;
 				ViewportClient.VisibilityDelegate = CachedViewportFlags.VisibilityDelegate;
 
+				FLevelEditorModule& LevelEditorModule = FModuleManager::GetModuleChecked<FLevelEditorModule>(TEXT("LevelEditor"));
+				LevelEditorModule.OnGetVisualizeCheckerboardEnabled().Unbind();
+				
 				if (TSharedPtr<FSceneViewport> LevelViewport = PinnedViewportInterface->GetSharedActiveViewport())
 				{
 					LevelViewport->SetFixedViewportSize(CachedViewportFlags.ViewportSize.X, CachedViewportFlags.ViewportSize.Y);
@@ -330,6 +339,7 @@ namespace MediaProfilePlaybackManager
 			bool bDrawAxes = false;
 			bool bDisableInput = false;
 			bool bAllowCinematicControl = false;
+			bool bDrawAlphaBlendedCheckerboard = false;
 			FEngineShowFlags EngineShowFlags = ESFIM_Editor;
 			FEngineShowFlags LastEngineShowFlags = ESFIM_Game;
 			FViewportStateGetter VisibilityDelegate;
@@ -778,7 +788,11 @@ UMediaCapture* UMediaProfilePlaybackManager::OpenManagedViewportOutputFromIndex(
 		return MediaCapture;
 	}
 
-	return nullptr;
+	// If the capture failed to start, mark the capture as errored, clean up its capture resources, and return it
+	CapturesWithError.Add(MediaCapture);
+	CleanUpCaptureResources(MediaOutput);
+	
+	return MediaCapture;
 }
 
 UMediaCapture* UMediaProfilePlaybackManager::OpenManagedViewportOutput(UMediaOutput* InMediaOutput, const FMediaCaptureOptions& InCaptureOptions, bool bAutoRestartCapture)
@@ -910,7 +924,11 @@ UMediaCapture* UMediaProfilePlaybackManager::OpenActiveViewportOutputFromIndex(i
 		return MediaCapture;
 	}
 
-	return nullptr;
+	// If the capture failed to start, mark the capture as errored, clean up its capture resources, and return it
+	CapturesWithError.Add(MediaCapture);
+	CleanUpCaptureResources(MediaOutput);
+	
+	return MediaCapture;
 }
 
 UMediaCapture* UMediaProfilePlaybackManager::OpenActiveViewportOutput(UMediaOutput* InMediaOutput, const FMediaCaptureOptions& InCaptureOptions, bool bAutoRestartCapture)
@@ -964,7 +982,11 @@ UMediaCapture* UMediaProfilePlaybackManager::OpenRenderTargetOutputFromIndex(int
 		return MediaCapture;
 	}
 
-	return nullptr;
+	// If the capture failed to start, mark the capture as errored, clean up its capture resources, and return it
+	CapturesWithError.Add(MediaCapture);
+	CleanUpCaptureResources(MediaOutput);
+	
+	return MediaCapture;
 }
 
 UMediaCapture* UMediaProfilePlaybackManager::OpenRenderTargetOutput(UMediaOutput* InMediaOutput, UTextureRenderTarget2D* InRenderTarget, const FMediaCaptureOptions& InCaptureOptions, bool bAutoRestartCapture)
@@ -1147,31 +1169,7 @@ void UMediaProfilePlaybackManager::PostStopCapture(int32 InMediaOutputIndex, FCl
 
 	if (InArgs.bDestroyCaptureObjects)
 	{
-		ManagedOutputViewports.RemoveAll([MediaOutput](const TSharedPtr<MediaProfilePlaybackManager::FManagedViewport>& InManagedViewport)
-		{
-			return InManagedViewport->MediaOutput == MediaOutput && InManagedViewport->Consumers.IsEmpty();
-		});
-
-		ActiveOutputViewports.RemoveAll([MediaOutput](const TSharedPtr<MediaProfilePlaybackManager::FActiveViewport>& InViewport)
-		{
-			return InViewport->MediaOutput == MediaOutput;
-		});
-
-#if WITH_EDITOR
-		if (ManagedOutputViewports.IsEmpty())
-		{
-			FEditorDelegates::PostPIEStarted.RemoveAll(this);
-			FEditorDelegates::PrePIEEnded.RemoveAll(this);
-		}
-
-		if (ActiveOutputViewports.IsEmpty())
-		{
-			if (GEditor)
-			{
-				GEditor->OnLevelViewportClientListChanged().RemoveAll(this);
-			}
-		}
-#endif
+		CleanUpCaptureResources(MediaOutput);
 
 		MediaCapture->OnStateChangedNative.RemoveAll(this);
 		MediaOutputCaptures.Remove(InMediaOutputIndex);
@@ -1423,6 +1421,29 @@ void UMediaProfilePlaybackManager::OnCaptureStateChanged(UMediaCapture* InMediaC
 	{
 		CapturesWithError.Remove(InMediaCapture);
 	}
+	else if (InMediaCapture->GetState() == EMediaCaptureState::Stopped)
+	{
+		// If the capture has been stopped due to an error, try and clean up any resources the media capture was using
+		if (CapturesWithError.Contains(InMediaCapture))
+		{
+			if (const int32* MediaOutputIndex = MediaOutputCaptures.FindKey(InMediaCapture))
+			{
+				UMediaProfile* MediaProfile = GetOwningMediaProfile();
+				if (!MediaProfile)
+				{
+					return;
+				}
+				
+				UMediaOutput* MediaOutput = MediaProfile->GetMediaOutput(*MediaOutputIndex);
+				if (!IsValid(MediaOutput))
+				{
+					return;
+				}
+				
+				CleanUpCaptureResources(MediaOutput);
+			}
+		}
+	}
 }
 
 void UMediaProfilePlaybackManager::QueueRestartCapture(UMediaOutput* InMediaOutput, FRecaptureFunc InRestartCaptureFunc)
@@ -1503,6 +1524,35 @@ void UMediaProfilePlaybackManager::Cleanup()
 #endif
 }
 
+void UMediaProfilePlaybackManager::CleanUpCaptureResources(UMediaOutput* InMediaOutput)
+{
+	ManagedOutputViewports.RemoveAll([InMediaOutput](const TSharedPtr<MediaProfilePlaybackManager::FManagedViewport>& InManagedViewport)
+		{
+			return InManagedViewport->MediaOutput == InMediaOutput && InManagedViewport->Consumers.IsEmpty();
+		});
+
+	ActiveOutputViewports.RemoveAll([InMediaOutput](const TSharedPtr<MediaProfilePlaybackManager::FActiveViewport>& InViewport)
+	{
+		return InViewport->MediaOutput == InMediaOutput;
+	});
+
+#if WITH_EDITOR
+	if (ManagedOutputViewports.IsEmpty())
+	{
+		FEditorDelegates::PostPIEStarted.RemoveAll(this);
+		FEditorDelegates::PrePIEEnded.RemoveAll(this);
+	}
+
+	if (ActiveOutputViewports.IsEmpty())
+	{
+		if (GEditor)
+		{
+			GEditor->OnLevelViewportClientListChanged().RemoveAll(this);
+		}
+	}
+#endif
+}
+
 #if WITH_EDITOR
 void UMediaProfilePlaybackManager::OnPostPIEStarted(bool bIsSimulating)
 {
diff --git a/Engine/Plugins/Media/MediaFrameworkUtilities/Source/MediaFrameworkUtilities/Public/Profile/MediaProfilePlaybackManager.h b/Engine/Plugins/Media/MediaFrameworkUtilities/Source/MediaFrameworkUtilities/Public/Profile/MediaProfilePlaybackManager.h
--- a/Engine/Plugins/Media/MediaFrameworkUtilities/Source/MediaFrameworkUtilities/Public/Profile/MediaProfilePlaybackManager.h
+++ b/Engine/Plugins/Media/MediaFrameworkUtilities/Source/MediaFrameworkUtilities/Public/Profile/MediaProfilePlaybackManager.h
@@ -370,6 +370,9 @@ private:
 	/** Cleans up all open sources and outputs, and deletes any managed state */
 	void Cleanup();
 
+	/** Cleans up any managed resources used by the specified media output, such as active viewports or managed viewports */
+	void CleanUpCaptureResources(UMediaOutput* InMediaOutput);
+	
 	/** Handles cleaning up resources and invoking the post StopCapture callback. */
 	void PostStopCapture(int32 InMediaOutputIndex, FCloseOutputArgs InArgs);
 	
diff --git a/Engine/Plugins/Media/MediaFrameworkUtilities/Source/MediaFrameworkUtilitiesEditor/Private/AssetEditor/SMediaProfileMediaItemDisplay.cpp b/Engine/Plugins/Media/MediaFrameworkUtilities/Source/MediaFrameworkUtilitiesEditor/Private/AssetEditor/SMediaProfileMediaItemDisplay.cpp
--- a/Engine/Plugins/Media/MediaFrameworkUtilities/Source/MediaFrameworkUtilitiesEditor/Private/AssetEditor/SMediaProfileMediaItemDisplay.cpp
+++ b/Engine/Plugins/Media/MediaFrameworkUtilities/Source/MediaFrameworkUtilitiesEditor/Private/AssetEditor/SMediaProfileMediaItemDisplay.cpp
@@ -1031,13 +1031,6 @@ void SMediaProfileMediaOutputDisplay::SetInvertAlphaChannelMask(bool bInInvertAl
 	if (TSharedPtr<SMediaCaptureImageBase> CaptureImageWidget = StaticCastSharedPtr<SMediaCaptureImageBase>(ImageWidget))
 	{
 		CaptureImageWidget->SetAlphaInverted(bInInvertAlphaChannelMask);
-
-		if (!bInInvertAlphaChannelMask && RenderTargetOutputInfoIndex == INDEX_NONE)
-		{
-			// For viewports, automatically disable the checkerboard if the alpha is set to non-inverted match the behavior
-			// of the level editor viewport RGBA feature.
-			SetDrawAlphaBlendedCheckerboard(false);
-		}
 	}
 }
 
@@ -1053,10 +1046,14 @@ void SMediaProfileMediaOutputDisplay::SetDrawAlphaBlendedCheckerboard(bool bInDr
 
 bool SMediaProfileMediaOutputDisplay::IsDrawAlphaBlendedCheckerboardEnabled() const
 {
-	// For viewports, don't allow checkerboard to be enabled if the alpha is not inverted to match the behavior
-	// of the level editor viewport RGBA feature. Allow it for render targets since render targets may have
-	// a standard alpha content instead of scene color
-	return RenderTargetOutputInfoIndex != INDEX_NONE ? true : GetInvertAlphaChannelMask();
+	// UE-354997: Don't allow alpha blended checkerboards for media outputs for now as they can inadvertently be passed to the media output stream
+	return false;
+}
+
+bool SMediaProfileMediaOutputDisplay::CanShowAlphaCheckerboard() const
+{
+	// UE-354997: Don't allow alpha blended checkerboards for media outputs for now as they can inadvertently be passed to the media output stream
+	return false;
 }
 
 void SMediaProfileMediaOutputDisplay::AddToolbarEntries(FToolMenuSection& Section)
diff --git a/Engine/Plugins/Media/MediaFrameworkUtilities/Source/MediaFrameworkUtilitiesEditor/Private/AssetEditor/SMediaProfileMediaItemDisplay.h b/Engine/Plugins/Media/MediaFrameworkUtilities/Source/MediaFrameworkUtilitiesEditor/Private/AssetEditor/SMediaProfileMediaItemDisplay.h
--- a/Engine/Plugins/Media/MediaFrameworkUtilities/Source/MediaFrameworkUtilitiesEditor/Private/AssetEditor/SMediaProfileMediaItemDisplay.h
+++ b/Engine/Plugins/Media/MediaFrameworkUtilities/Source/MediaFrameworkUtilitiesEditor/Private/AssetEditor/SMediaProfileMediaItemDisplay.h
@@ -518,7 +518,7 @@ protected:
 				);
 				InvertAlphaEntry.Visibility = TAttribute<bool>::CreateSP(this, &SMediaProfileMediaItemDisplayBase::CanInvertAlphaChannelMask);
 
-				OptionsSubMenuSection.AddMenuEntry(
+				FToolMenuEntry& AlphaCheckerboardEntry = OptionsSubMenuSection.AddMenuEntry(
 					"VisualizeAlphaWithCheckerboard",
 					LOCTEXT("VisualizeWithCheckerboardLabel", "Visualize Alpha with Checkerboard"),
 					LOCTEXT("VisualizeWithCheckerboardToolTip", "Displays a checkerboard underneath the viewport to help visualize the alpha channel"),
@@ -539,6 +539,7 @@ protected:
 					),
 					EUserInterfaceActionType::ToggleButton
 				);
+				AlphaCheckerboardEntry.Visibility = TAttribute<bool>::CreateSP(this, &SMediaProfileMediaItemDisplayBase::CanShowAlphaCheckerboard);
 			}),
 			/*bInOpenSubMenuOnClick=*/false,
 			TAttribute<FSlateIcon>::CreateLambda(
@@ -872,6 +873,12 @@ protected:
 		return true;
 	}
 
+	/** Indicates whether the media item display can show the 'Visualize Alpha with Checkerboard' option in the toolbar dropdown menu */
+	virtual bool CanShowAlphaCheckerboard() const
+	{
+		return true;	
+	}
+	
 	/** Gets whether this panel is currently maximized in the viewport */
 	bool IsMaximized() const
 	{
@@ -1044,6 +1051,7 @@ protected:
 	virtual bool CanInvertAlphaChannelMask() const override { return true; }
 	virtual void SetDrawAlphaBlendedCheckerboard(bool bInDrawAlphaBlendedCheckerboard) override;
 	virtual bool IsDrawAlphaBlendedCheckerboardEnabled() const override;
+	virtual bool CanShowAlphaCheckerboard() const override;
 
 	virtual void AddToolbarEntries(FToolMenuSection& Section) override;
 	virtual void AppendToolbarCommandList(FToolMenuContext& Context) override;
diff --git a/Engine/Source/Editor/LevelEditor/Private/ViewportToolbar/LevelEditorSubmenus.cpp b/Engine/Source/Editor/LevelEditor/Private/ViewportToolbar/LevelEditorSubmenus.cpp
--- a/Engine/Source/Editor/LevelEditor/Private/ViewportToolbar/LevelEditorSubmenus.cpp
+++ b/Engine/Source/Editor/LevelEditor/Private/ViewportToolbar/LevelEditorSubmenus.cpp
@@ -1028,12 +1028,16 @@ FToolMenuEntry CreateChannelsSubmenu(const TWeakPtr<::SLevelViewport>& LevelView
 					}),
 					FCanExecuteAction::CreateLambda([LevelViewportWeak]
 					{
+						FLevelEditorModule& LevelEditorModule = FModuleManager::GetModuleChecked<FLevelEditorModule>(TEXT("LevelEditor"));
+						const bool bEnabled = LevelEditorModule.OnGetVisualizeCheckerboardEnabled().IsBound() ?
+							LevelEditorModule.OnGetVisualizeCheckerboardEnabled().Execute() : true;
+						
 						if (const TSharedPtr<::SLevelViewport> Viewport = LevelViewportWeak.Pin())
 						{
-							return Viewport->GetViewportClient()->ChannelMaskParams.bInvertAlphaChannelMask;
+							return bEnabled && Viewport->GetViewportClient()->ChannelMaskParams.bInvertAlphaChannelMask;
 						}
 
-						return false;
+						return bEnabled;
 					}),
 					FGetActionCheckState::CreateLambda([LevelViewportWeak]
 					{
diff --git a/Engine/Source/Editor/LevelEditor/Public/LevelEditor.h b/Engine/Source/Editor/LevelEditor/Public/LevelEditor.h
--- a/Engine/Source/Editor/LevelEditor/Public/LevelEditor.h
+++ b/Engine/Source/Editor/LevelEditor/Public/LevelEditor.h
@@ -339,6 +339,10 @@ public:
 	/** Delegate used to capture skeletal meshes to single-frame animations when 'keeping simulation changes' */
 	DECLARE_DELEGATE_RetVal_OneParam(UAnimSequence*, FCaptureSingleFrameAnimSequence, USkeletalMeshComponent* /*Component*/);
 	virtual FCaptureSingleFrameAnimSequence& OnCaptureSingleFrameAnimSequence() { return CaptureSingleFrameAnimSequenceDelegate; }
+
+	/** Delegate used to determine if users should be able to enable the alpha-blended checkerboard for the level viewport */
+	DECLARE_DELEGATE_RetVal(bool, FGetVisualizeCheckerboardEnabled);
+	virtual FGetVisualizeCheckerboardEnabled& OnGetVisualizeCheckerboardEnabled() { return GetVisualizeCheckerboardEnabled; }
 	
 	/**  Add a custom filter to the outliner filter bar. These are all AND'd together
 	 * @see FLevelEditorOutlinerSettings
@@ -499,6 +503,9 @@ private:
 	/** Delegate used to capture skeletal meshes to single-frame animations when 'keeping simulation changes' */
 	FCaptureSingleFrameAnimSequence CaptureSingleFrameAnimSequenceDelegate;
 
+	/** Delegate used to determine if 'Visualize Alpha with Checkerboard' checkbox is enabled in the level viewports  */
+	FGetVisualizeCheckerboardEnabled GetVisualizeCheckerboardEnabled;
+	
 	/** All extender delegates for the level viewport menus */
 	TArray<FLevelViewportMenuExtender_SelectedObjects> LevelViewportDragDropContextMenuExtenders;
 	TArray<FLevelViewportMenuExtender_SelectedActors> LevelViewportContextMenuExtenders;
