Re-write of the description of the meda framework plugin. Added media profile to the description to make it easier to find the plugin that hosts the media profile editor.

#rnx
#Synced-CL 50213209

diff --git a/Engine/Plugins/Media/MediaFrameworkUtilities/MediaFrameworkUtilities.uplugin b/Engine/Plugins/Media/MediaFrameworkUtilities/MediaFrameworkUtilities.uplugin
--- a/Engine/Plugins/Media/MediaFrameworkUtilities/MediaFrameworkUtilities.uplugin
+++ b/Engine/Plugins/Media/MediaFrameworkUtilities/MediaFrameworkUtilities.uplugin
@@ -3,7 +3,7 @@
 	"Version": 1,
 	"VersionName": "1.0",
 	"FriendlyName": "Media Framework Utilities",
-	"Description": "Utility assets and actors to ease the use of the Media Framework.",
+	"Description": "This plugin provides utility assets and actors designed to simplify the Media Framework setup. It includes access to the the Media Profile editor.",
 	"Category": "Media Players",
 	"CreatedBy": "Epic Games, Inc.",
 	"CreatedByURL": "https://epicgames.com",
@@ -51,4 +51,4 @@
 	],
 	"DocsURL": "",
 	"IsExperimentalVersion": false
-}
\ No newline at end of file
+}
