Media Profile: Added media player options to media profile editor details panels for stream media source.

#Synced-CL 51110937

#Synced-CL 51110937

diff --git a/Engine/Plugins/Media/MediaFrameworkUtilities/Source/MediaFrameworkUtilitiesEditor/Private/AssetEditor/MediaProfileMediaItemDetailCustomization.cpp b/Engine/Plugins/Media/MediaFrameworkUtilities/Source/MediaFrameworkUtilitiesEditor/Private/AssetEditor/MediaProfileMediaItemDetailCustomization.cpp
--- a/Engine/Plugins/Media/MediaFrameworkUtilities/Source/MediaFrameworkUtilitiesEditor/Private/AssetEditor/MediaProfileMediaItemDetailCustomization.cpp
+++ b/Engine/Plugins/Media/MediaFrameworkUtilities/Source/MediaFrameworkUtilitiesEditor/Private/AssetEditor/MediaProfileMediaItemDetailCustomization.cpp
@@ -3,16 +3,38 @@
 #include "MediaProfileMediaItemDetailCustomization.h"
 
 #include "IDetailChildrenBuilder.h"
+#include "IDetailGroup.h"
 #include "IPropertyUtilities.h"
 #include "IStructureDataProvider.h"
 #include "MediaFrameworkWorldSettingsAssetUserData.h"
+#include "MediaPlayer.h"
 #include "MediaProfileEditorUserSettings.h"
+#include "MediaTexture.h"
+#include "StreamMediaSource.h"
 #include "Algo/Count.h"
+#include "Internationalization/Culture.h"
 #include "Profile/MediaProfilePlaybackManager.h"
 #include "UI/MediaFrameworkUtilitiesEditorStyle.h"
+#include "Widgets/Input/SComboBox.h"
 
 #define LOCTEXT_NAMESPACE "FMediaProfileMediaItemDetailCustomization"
 
+void FMediaProfileMediaSourceDetailCustomization::CustomizeDetails(IDetailLayoutBuilder& DetailBuilder)
+{
+	FMediaProfileMediaItemDetailCustomization<UMediaSource>::CustomizeDetails(DetailBuilder);
+
+	TArray<TWeakObjectPtr<UMediaSource>> Sources = CachedDetailBuilder.Pin()->GetObjectsOfTypeBeingCustomized<UMediaSource>();
+
+	if (Sources.Num() == 1)
+	{
+		TStrongObjectPtr<UMediaSource> Source = Sources[0].Pin();
+		if (Source.IsValid() && ShowMediaPlayerOptions(Source.Get()))
+		{
+			AddMediaPlayerProperties(Source.Get(), DetailBuilder);
+		}
+	}
+}
+
 FText FMediaProfileMediaSourceDetailCustomization::GetMediaItemTypeText() const
 {
 	return LOCTEXT("MediaSourceTypeText", "Media Source");
@@ -81,6 +103,482 @@ void FMediaProfileMediaSourceDetailCustomization::SetMediaObject(UMediaSource* I
 	MediaProfile->SetMediaSource(Index, InNewMediaObject);
 }
 
+bool FMediaProfileMediaSourceDetailCustomization::ShowMediaPlayerOptions(UMediaSource* InSource) const
+{
+	// TODO: Might need something more sophisticated in the future (something queryable on the source type, like a class attribute)
+	// if more media source types need to display the media player properties
+	if (InSource->IsA<UStreamMediaSource>())
+	{
+		return true;
+	}
+
+	return false;
+}
+
+class SIndexComboBox : public SComboBox<TSharedPtr<int32>>
+{
+public:
+
+	DECLARE_DELEGATE_RetVal_OneParam(FText, FGetIndexDisplayText, const int32);
+	DECLARE_DELEGATE_TwoParams(FOnIndexChanged, int32, ESelectInfo::Type);
+	
+	SLATE_BEGIN_ARGS(SIndexComboBox) { }
+		SLATE_ATTRIBUTE(int32, NumIndices)
+		SLATE_ATTRIBUTE(int32, StartingIndex)
+		SLATE_ATTRIBUTE(FSlateFontInfo, Font)
+		SLATE_ATTRIBUTE(int32, SelectedIndex)
+		SLATE_EVENT(FGetIndexDisplayText, GetIndexDisplayText)
+		SLATE_EVENT(FOnIndexChanged, OnIndexChanged)
+	SLATE_END_ARGS()
+
+	void Construct(const FArguments& InArgs)
+	{
+		NumIndicesAttr = InArgs._NumIndices;
+		StartingIndexAttr = InArgs._StartingIndex;
+		SelectedIndexAttr = InArgs._SelectedIndex;
+		GetIndexDisplayText = InArgs._GetIndexDisplayText;
+		OnIndexChanged = InArgs._OnIndexChanged;
+		
+		SComboBox<TSharedPtr<int32>>::FArguments BaseArgs;
+		BaseArgs.OptionsSource(&Indices);
+		BaseArgs.OnGenerateWidget_Lambda([WeakThis = AsWeak()](TSharedPtr<int32> Item)->TSharedRef<SWidget>
+		{
+			TSharedPtr<SIndexComboBox> This = StaticCastSharedPtr<SIndexComboBox>(WeakThis.Pin());
+			if (!This.IsValid())
+			{
+				return SNullWidget::NullWidget;
+			}
+			
+			const int32 Index = *Item;
+			return SNew(STextBlock)
+				.Text(TAttribute<FText>::CreateLambda([WeakThis, Index]
+				{
+					TSharedPtr<SIndexComboBox> This = StaticCastSharedPtr<SIndexComboBox>(WeakThis.Pin());
+					if (!This.IsValid())
+					{
+						return FText::GetEmpty();
+					}
+					
+					if (This->GetIndexDisplayText.IsBound())
+					{
+						return This->GetIndexDisplayText.Execute(Index);
+					}
+
+					return FText::AsNumber(Index);
+				}))
+				.Font(This->FontAttr);
+		});
+		BaseArgs.OnSelectionChanged_Lambda([WeakThis = AsWeak()](TSharedPtr<int32> Item, ESelectInfo::Type SelectType)
+		{
+			TSharedPtr<SIndexComboBox> This = StaticCastSharedPtr<SIndexComboBox>(WeakThis.Pin());
+			if (!This.IsValid())
+			{
+				return;
+			}
+			
+			if (!Item.IsValid())
+			{
+				return;
+			}
+			
+			This->OnIndexChanged.ExecuteIfBound(*Item, SelectType);
+		});
+		BaseArgs.Content()
+		[
+			SNew(STextBlock)
+			.Text(TAttribute<FText>::CreateLambda([WeakThis = AsWeak()]
+			{
+				TSharedPtr<SIndexComboBox> This = StaticCastSharedPtr<SIndexComboBox>(WeakThis.Pin());
+				if (!This.IsValid())
+				{
+					return FText::GetEmpty();
+				}
+				
+				if (This->Indices.IsEmpty())
+				{
+					return LOCTEXT("NoneLabel", "None");
+				}
+				
+				const int32 SelectedIndex = This->SelectedIndexAttr.Get(INDEX_NONE);
+				if (This->GetIndexDisplayText.IsBound())
+				{
+					return This->GetIndexDisplayText.Execute(SelectedIndex);
+				}
+
+				return FText::AsNumber(SelectedIndex);
+			}))
+			.Font(FontAttr)
+		];
+
+		SComboBox<TSharedPtr<int32>>::Construct(BaseArgs);
+
+		FillIndexList();
+	}
+
+	virtual void Tick(const FGeometry& AllottedGeometry, const double InCurrentTime, const float InDeltaTime) override
+	{
+		SComboBox<TSharedPtr<int32>>::Tick(AllottedGeometry, InCurrentTime, InDeltaTime);
+
+		const int32 StartingIndex = StartingIndexAttr.Get(0);
+        const int32 NumIndices = NumIndicesAttr.Get(0);
+
+		if (Indices.Num() != NumIndices || (Indices.IsValidIndex(0) && *Indices[0] != StartingIndex))
+		{
+			FillIndexList();
+		}
+
+		if (Indices.Num() > 0)
+		{
+			const int32 DesiredSelection = SelectedIndexAttr.Get(INDEX_NONE);
+			const int32 CurrentSelection = GetSelectedItem().IsValid() ? *GetSelectedItem() : INDEX_NONE;
+			if (DesiredSelection != CurrentSelection)
+			{
+				TSharedPtr<int32>* ItemToSelect = Indices.FindByPredicate([DesiredSelection](const TSharedPtr<int32>& Item)
+				{
+					return Item.IsValid() && *Item == DesiredSelection;
+				});
+
+				if (ItemToSelect)
+				{
+					SetSelectedItem(*ItemToSelect);
+				}
+				else
+				{
+					ClearSelection();
+				}
+			}
+		}
+		else
+		{
+			ClearSelection();
+		}
+	}
+	
+private:
+	void FillIndexList()
+	{
+		Indices.Empty();
+
+		const int32 StartingIndex = StartingIndexAttr.Get(0);
+		const int32 NumIndices = NumIndicesAttr.Get(0);
+		for (int32 Index = 0; Index < NumIndices; ++Index)
+		{
+			Indices.Add(MakeShared<int32>(Index + StartingIndex));
+		}
+
+		RefreshOptions();
+		
+		const int32 IndexToSelect = SelectedIndexAttr.Get(INDEX_NONE);
+
+		TSharedPtr<int32>* ItemToSelect = Indices.FindByPredicate([IndexToSelect](const TSharedPtr<int32>& Item)
+		{
+			return Item.IsValid() && *Item == IndexToSelect;
+		});
+
+		if (ItemToSelect)
+		{
+			SetSelectedItem(*ItemToSelect);
+		}
+		else
+		{
+			ClearSelection();
+		}
+	}
+	
+private:
+	TArray<TSharedPtr<int32>> Indices;
+
+	TAttribute<int32> NumIndicesAttr;
+	TAttribute<int32> StartingIndexAttr;
+	TAttribute<int32> SelectedIndexAttr;
+	TAttribute<FSlateFontInfo> FontAttr;
+	
+	FGetIndexDisplayText GetIndexDisplayText;
+	FOnIndexChanged OnIndexChanged;
+};
+
+void FMediaProfileMediaSourceDetailCustomization::AddMediaPlayerProperties(UMediaSource* InMediaSource, IDetailLayoutBuilder& DetailBuilder)
+{
+	UMediaProfile* ParentProfile = InMediaSource->GetTypedOuter<UMediaProfile>();
+	if (!ParentProfile)
+	{
+		return;
+	}
+
+	UMediaTexture* MediaTexture = ParentProfile->GetPlaybackManager()->GetSourceMediaTexture(InMediaSource);
+	if (!MediaTexture)
+	{
+		return;
+	}
+
+	SourceMediaPlayer = MediaTexture->GetMediaPlayer();
+	if (!SourceMediaPlayer.IsValid())
+	{
+		return;
+	}
+
+	TStrongObjectPtr<UMediaPlayer> PinnedMediaPlayer = SourceMediaPlayer.Pin();
+	
+	IDetailCategoryBuilder& PlaybackOptionsCategory = DetailBuilder.EditCategory(TEXT("PlaybackOptionsCategory"), LOCTEXT("PlaybackOptionsCategoryName", "Playback Options"));
+
+	const EMediaPlayerTrack TrackTypes[] =
+	{
+		EMediaPlayerTrack::Audio,
+		EMediaPlayerTrack::Caption,
+		EMediaPlayerTrack::Subtitle,
+		EMediaPlayerTrack::Text,
+		EMediaPlayerTrack::Video
+	};
+	
+	constexpr bool bForAdvanced = false;
+	constexpr bool bStartExpanded = true;
+	
+	IDetailGroup& TracksGroup = PlaybackOptionsCategory.AddGroup(TEXT("TracksGroup"), LOCTEXT("TracksGroupName", "Tracks"), bForAdvanced, bStartExpanded);
+
+	for (EMediaPlayerTrack TrackType : TrackTypes)
+	{
+		TracksGroup.AddWidgetRow()
+			.NameContent()
+			[
+				SNew(STextBlock)
+				.Text(StaticEnum<EMediaPlayerTrack>()->GetDisplayNameTextByIndex((int32)TrackType))
+				.Font(DetailBuilder.GetDetailFont())
+			]
+			.ValueContent()
+			[
+				SNew(SIndexComboBox)
+				.NumIndices(TAttribute<int32>::CreateLambda([WeakThis = AsWeak(), TrackType]()->int32
+				{
+					TSharedPtr<FMediaProfileMediaSourceDetailCustomization> This =
+						StaticCastSharedPtr<FMediaProfileMediaSourceDetailCustomization>(WeakThis.Pin());
+
+					if (!This.IsValid())
+					{
+						return 0;
+					}
+					
+					if (TStrongObjectPtr<UMediaPlayer> PinnedMediaPlayer = This->SourceMediaPlayer.Pin())
+					{
+						const int32 NumTracks = PinnedMediaPlayer->GetNumTracks(TrackType);
+						return NumTracks > 0 ? NumTracks + 1 : 0;
+					}
+
+					return 0;
+				}))
+				.StartingIndex(INDEX_NONE)
+				.SelectedIndex_Lambda([WeakThis = AsWeak(), TrackType]()->int32
+				{
+					TSharedPtr<FMediaProfileMediaSourceDetailCustomization> This =
+						StaticCastSharedPtr<FMediaProfileMediaSourceDetailCustomization>(WeakThis.Pin());
+
+					if (!This.IsValid())
+					{
+						return INDEX_NONE;
+					}
+					
+					if (TStrongObjectPtr<UMediaPlayer> PinnedMediaPlayer = This->SourceMediaPlayer.Pin())
+					{
+						return PinnedMediaPlayer->GetSelectedTrack(TrackType);
+					}
+
+					return INDEX_NONE;
+				})
+				.GetIndexDisplayText_Lambda([WeakThis = AsWeak(), TrackType](int32 TrackIndex)
+				{
+					TSharedPtr<FMediaProfileMediaSourceDetailCustomization> This =
+						StaticCastSharedPtr<FMediaProfileMediaSourceDetailCustomization>(WeakThis.Pin());
+
+					if (!This.IsValid())
+					{
+						return FText::GetEmpty();
+					}
+					
+					if (TrackIndex == INDEX_NONE)
+					{
+						return LOCTEXT("DisabledLabel", "Disabled");
+					}
+					
+					if (TStrongObjectPtr<UMediaPlayer> PinnedMediaPlayer = This->SourceMediaPlayer.Pin())
+					{
+						FInternationalization& I18n = FInternationalization::Get();
+						
+						const FText DisplayName = PinnedMediaPlayer->GetTrackDisplayName(TrackType, TrackIndex);
+						const FString Language = PinnedMediaPlayer->GetTrackLanguage(TrackType, TrackIndex);
+						const FCulturePtr Culture = I18n.GetCulture(Language);
+						const FString LanguageDisplayName = Culture.IsValid() ? Culture->GetDisplayName() : FString();
+						const FString LanguageNativeName = Culture.IsValid() ? Culture->GetNativeName() : FString();
+			
+						return LanguageNativeName.IsEmpty() ? DisplayName : FText::Format(LOCTEXT("TrackNameFormat", "{0} ({1})"), DisplayName, FText::FromString(LanguageNativeName));
+					}
+
+					return LOCTEXT("NoneLabel", "None");
+				})
+				.OnIndexChanged_Lambda([WeakThis = AsWeak(), TrackType](int32 TrackIndex, ESelectInfo::Type SelectType)
+				{
+					TSharedPtr<FMediaProfileMediaSourceDetailCustomization> This =
+						StaticCastSharedPtr<FMediaProfileMediaSourceDetailCustomization>(WeakThis.Pin());
+
+					if (!This.IsValid())
+					{
+						return;
+					}
+					
+					if (TStrongObjectPtr<UMediaPlayer> PinnedMediaPlayer = This->SourceMediaPlayer.Pin())
+					{
+						PinnedMediaPlayer->SelectTrack(TrackType, TrackIndex);
+					}
+				})
+			];
+	}
+	
+	IDetailGroup& FormatsGroup = PlaybackOptionsCategory.AddGroup(TEXT("FormatsGroup"), LOCTEXT("FormatsGroupName", "Formats"), bForAdvanced, bStartExpanded);
+
+	for (EMediaPlayerTrack TrackType : TrackTypes)
+	{
+		FormatsGroup.AddWidgetRow()
+			.NameContent()
+			[
+				SNew(STextBlock)
+				.Text(StaticEnum<EMediaPlayerTrack>()->GetDisplayNameTextByIndex((int32)TrackType))
+				.Font(DetailBuilder.GetDetailFont())
+			]
+			.ValueContent()
+			[
+				SNew(SIndexComboBox)
+				.NumIndices(TAttribute<int32>::CreateLambda([WeakThis = AsWeak(), TrackType]()->int32
+				{
+					TSharedPtr<FMediaProfileMediaSourceDetailCustomization> This =
+						StaticCastSharedPtr<FMediaProfileMediaSourceDetailCustomization>(WeakThis.Pin());
+
+					if (!This.IsValid())
+					{
+						return 0;
+					}
+					
+					if (TStrongObjectPtr<UMediaPlayer> PinnedMediaPlayer = This->SourceMediaPlayer.Pin())
+					{
+						const int32 SelectedTrack = PinnedMediaPlayer->GetSelectedTrack(TrackType);
+						const int32 NumFormats = PinnedMediaPlayer->GetNumTrackFormats(TrackType, SelectedTrack);
+						return SelectedTrack != INDEX_NONE ? NumFormats : 0;
+					}
+
+					return 0;
+				}))
+				.SelectedIndex_Lambda([WeakThis = AsWeak(), TrackType]()->int32
+				{
+					TSharedPtr<FMediaProfileMediaSourceDetailCustomization> This =
+						StaticCastSharedPtr<FMediaProfileMediaSourceDetailCustomization>(WeakThis.Pin());
+
+					if (!This.IsValid())
+					{
+						return INDEX_NONE;
+					}
+					
+					if (TStrongObjectPtr<UMediaPlayer> PinnedMediaPlayer = This->SourceMediaPlayer.Pin())
+					{
+						const int32 SelectedTrack = PinnedMediaPlayer->GetSelectedTrack(TrackType);
+						return SelectedTrack != INDEX_NONE ? PinnedMediaPlayer->GetTrackFormat(TrackType, SelectedTrack) : INDEX_NONE;
+					}
+
+					return INDEX_NONE;
+				})
+				.GetIndexDisplayText_Lambda([WeakThis = AsWeak(), TrackType](int32 FormatIndex)
+				{
+					TSharedPtr<FMediaProfileMediaSourceDetailCustomization> This =
+						StaticCastSharedPtr<FMediaProfileMediaSourceDetailCustomization>(WeakThis.Pin());
+
+					if (!This.IsValid())
+					{
+						return FText::GetEmpty();
+					}
+					
+					if (FormatIndex == INDEX_NONE)
+					{
+						return LOCTEXT("DisabledLabel", "Disabled");
+					}
+					
+					if (TStrongObjectPtr<UMediaPlayer> PinnedMediaPlayer = This->SourceMediaPlayer.Pin())
+					{
+						const int32 SelectedTrack = PinnedMediaPlayer->GetSelectedTrack(TrackType);
+						if (SelectedTrack == INDEX_NONE)
+						{
+							return LOCTEXT("NoneLabel", "None");
+						}
+						
+						FText DisplayText;
+						
+						if (TrackType == EMediaPlayerTrack::Audio)
+						{
+							const uint32 Channels = PinnedMediaPlayer->GetAudioTrackChannels(SelectedTrack, FormatIndex);
+							const uint32 SampleRate = PinnedMediaPlayer->GetAudioTrackSampleRate(SelectedTrack, FormatIndex);
+							const FString Type = PinnedMediaPlayer->GetAudioTrackType(SelectedTrack, FormatIndex);
+
+							DisplayText = FText::Format(LOCTEXT("TrackFormatMenuAudioFormat", "{0}: {1} {2} channels @ {3} Hz"),
+								FText::AsNumber(FormatIndex),
+								FText::FromString(Type),
+								FText::AsNumber(Channels),
+								FText::AsNumber(SampleRate)
+							);
+						}
+						else if (TrackType == EMediaPlayerTrack::Video)
+						{
+							const FIntPoint Dim = PinnedMediaPlayer->GetVideoTrackDimensions(SelectedTrack, FormatIndex);
+							const float FrameRate = PinnedMediaPlayer->GetVideoTrackFrameRate(SelectedTrack, FormatIndex);
+							const TRange<float> FrameRates = PinnedMediaPlayer->GetVideoTrackFrameRates(SelectedTrack, FormatIndex);
+							const FString Type = PinnedMediaPlayer->GetVideoTrackType(SelectedTrack, FormatIndex);
+
+							FFormatNamedArguments Arguments;
+							Arguments.Add(TEXT("Index"), FText::AsNumber(FormatIndex));
+							Arguments.Add(TEXT("DimX"), FText::AsNumber(Dim.X));
+							Arguments.Add(TEXT("DimY"), FText::AsNumber(Dim.Y));
+							Arguments.Add(TEXT("Fps"), FText::AsNumber(FrameRate));
+							Arguments.Add(TEXT("Type"), FText::FromString(Type));
+
+							if (FrameRates.IsDegenerate() && (FrameRates.GetLowerBoundValue() == FrameRate))
+							{
+								DisplayText = FText::Format(LOCTEXT("TrackFormatMenuVideoFormat", "{Index}: {Type} {DimX}x{DimY} {Fps} fps"), Arguments);
+							}
+							else
+							{
+								Arguments.Add(TEXT("FpsLower"), FText::AsNumber(FrameRates.GetLowerBoundValue()));
+								Arguments.Add(TEXT("FpsUpper"), FText::AsNumber(FrameRates.GetUpperBoundValue()));
+
+								DisplayText = FText::Format(LOCTEXT("TrackFormatMenuVideoFormat2", "{Index}: {Type} {DimX}x{DimY} {Fps} [{FpsLower}-{FpsUpper}] fps"), Arguments);
+							}
+						}
+						else
+						{
+							DisplayText = LOCTEXT("TrackFormatDefault", "Default");
+						}
+
+						return DisplayText;
+					}
+					
+					return LOCTEXT("NoneLabel", "None");
+				})
+				.OnIndexChanged_Lambda([WeakThis = AsWeak(), TrackType](int32 FormatIndex, ESelectInfo::Type SelectType)
+				{
+					TSharedPtr<FMediaProfileMediaSourceDetailCustomization> This =
+						StaticCastSharedPtr<FMediaProfileMediaSourceDetailCustomization>(WeakThis.Pin());
+
+					if (!This.IsValid())
+					{
+						return;
+					}
+					
+					if (TStrongObjectPtr<UMediaPlayer> PinnedMediaPlayer = This->SourceMediaPlayer.Pin())
+					{
+						const int32 SelectedTrack = PinnedMediaPlayer->GetSelectedTrack(TrackType);
+						if (SelectedTrack != INDEX_NONE)
+						{
+							PinnedMediaPlayer->SetTrackFormat(TrackType, SelectedTrack, FormatIndex);
+						}
+					}
+				})
+			];
+	}
+}
+
 class SMediaCaptureMethodComboButton : public SComboButton
 {
 public:
diff --git a/Engine/Plugins/Media/MediaFrameworkUtilities/Source/MediaFrameworkUtilitiesEditor/Private/AssetEditor/MediaProfileMediaItemDetailCustomization.h b/Engine/Plugins/Media/MediaFrameworkUtilities/Source/MediaFrameworkUtilitiesEditor/Private/AssetEditor/MediaProfileMediaItemDetailCustomization.h
--- a/Engine/Plugins/Media/MediaFrameworkUtilities/Source/MediaFrameworkUtilitiesEditor/Private/AssetEditor/MediaProfileMediaItemDetailCustomization.h
+++ b/Engine/Plugins/Media/MediaFrameworkUtilities/Source/MediaFrameworkUtilitiesEditor/Private/AssetEditor/MediaProfileMediaItemDetailCustomization.h
@@ -22,6 +22,7 @@
 
 #define LOCTEXT_NAMESPACE "FMediaProfileMediaItemDetailCustomization"
 
+class UMediaPlayer;
 class FStructOnScopeStructureDataProvider;
 
 /**
@@ -336,12 +337,24 @@ public:
 		: FMediaProfileMediaItemDetailCustomization(InMediaProfileEditor, InMediaProfile, InOnObjectsChanged)
 	{ }
 
+	virtual void CustomizeDetails(IDetailLayoutBuilder& DetailBuilder) override;
+	
 private:
 	virtual FText GetMediaItemTypeText() const override;
 	virtual void RegisterMediaTypeSection(FPropertyEditorModule& PropertyModule, const FName& MediaTypeCategory) override;
 	virtual FText GetMediaObjectLabel(UMediaSource* InMediaObject) const override;
 	virtual void SetMediaObjectLabel(UMediaSource* InMediaObject, const FText& InLabel) override;
 	virtual void SetMediaObject(UMediaSource* InOriginalMediaObject, UMediaSource* InNewMediaObject) override;
+
+	/** Gets whether the properties of the media player that is playing the media source should be displayed in the details panel */
+	bool ShowMediaPlayerOptions(UMediaSource* InSource) const;
+
+	/** Adds appropraite media player properties to the details panel of the media source being displayed */
+	void AddMediaPlayerProperties(UMediaSource* InMediaSource, IDetailLayoutBuilder& DetailBuilder);
+
+private:
+	/** Weak pointer to the media player of the media source being customized */
+	TWeakObjectPtr<UMediaPlayer> SourceMediaPlayer;
 };
 
 /**
diff --git a/Engine/Plugins/Media/WmfMedia/Source/WmfMedia/Private/Player/WmfMediaTracks.cpp b/Engine/Plugins/Media/WmfMedia/Source/WmfMedia/Private/Player/WmfMediaTracks.cpp
--- a/Engine/Plugins/Media/WmfMedia/Source/WmfMedia/Private/Player/WmfMediaTracks.cpp
+++ b/Engine/Plugins/Media/WmfMedia/Source/WmfMedia/Private/Player/WmfMediaTracks.cpp
@@ -715,24 +715,28 @@ int32 FWmfMediaTracks::GetNumTrackFormats(EMediaTrackType TrackType, int32 Track
 		{
 			return AudioTracks[TrackIndex].Formats.Num();
 		}
+		break;
 
 	case EMediaTrackType::Metadata:
 		if (MetadataTracks.IsValidIndex(TrackIndex))
 		{
 			return 1;
 		}
+		break;
 
 	case EMediaTrackType::Caption:
 		if (CaptionTracks.IsValidIndex(TrackIndex))
 		{
 			return 1;
 		}
+		break;
 
 	case EMediaTrackType::Video:
 		if (VideoTracks.IsValidIndex(TrackIndex))
 		{
 			return VideoTracks[TrackIndex].Formats.Num();
 		}
+		break;
 
 	default:
 		break; // unsupported track type
