Media Profile: Change default behavior for Immersive View option in media profile editor to support immersive view of multiple panels instead of just single viewport

#Synced-CL 51181393

diff --git a/Engine/Plugins/Media/MediaFrameworkUtilities/Source/MediaFrameworkUtilitiesEditor/Private/AssetEditor/SMediaProfileMediaItemDisplay.h b/Engine/Plugins/Media/MediaFrameworkUtilities/Source/MediaFrameworkUtilitiesEditor/Private/AssetEditor/SMediaProfileMediaItemDisplay.h
--- a/Engine/Plugins/Media/MediaFrameworkUtilities/Source/MediaFrameworkUtilitiesEditor/Private/AssetEditor/SMediaProfileMediaItemDisplay.h
+++ b/Engine/Plugins/Media/MediaFrameworkUtilities/Source/MediaFrameworkUtilitiesEditor/Private/AssetEditor/SMediaProfileMediaItemDisplay.h
@@ -187,7 +187,7 @@ protected:
 						
 					if (!IsImmersive())
 					{
-						PinnedOwningViewport->SetImmersivePanel(PanelIndex);
+						PinnedOwningViewport->SetImmersivePanel();
 					}
 					else
 					{
@@ -900,7 +900,7 @@ protected:
 			return false;
 		}
 
-		return PinnedOwningViewport->IsPanelImmersive(PanelIndex);
+		return PinnedOwningViewport->IsPanelImmersive();
 	}
 
 	/** Changes the panel's contents, recreating the display widgets */
diff --git a/Engine/Plugins/Media/MediaFrameworkUtilities/Source/MediaFrameworkUtilitiesEditor/Private/AssetEditor/SMediaProfileViewport.cpp b/Engine/Plugins/Media/MediaFrameworkUtilities/Source/MediaFrameworkUtilitiesEditor/Private/AssetEditor/SMediaProfileViewport.cpp
--- a/Engine/Plugins/Media/MediaFrameworkUtilities/Source/MediaFrameworkUtilitiesEditor/Private/AssetEditor/SMediaProfileViewport.cpp
+++ b/Engine/Plugins/Media/MediaFrameworkUtilities/Source/MediaFrameworkUtilitiesEditor/Private/AssetEditor/SMediaProfileViewport.cpp
@@ -186,7 +186,6 @@ bool SMediaProfileViewport::IsPanelLocked(int32 InPanelIndex) const
 bool SMediaProfileViewport::CanMaximizePanel(int32 InPanelIndex) const
 {
 	return MaximizedPanelIndex == INDEX_NONE &&
-		ImmersivePanelIndex == INDEX_NONE &&
 		PanelContents.Num() > 1 &&
 		PanelContents.IsValidIndex(InPanelIndex);
 }
@@ -224,11 +223,6 @@ bool SMediaProfileViewport::CanRestorePreviousLayout() const
 
 void SMediaProfileViewport::RestorePreviousLayout()
 {
-	if (ImmersivePanelIndex != INDEX_NONE)
-	{
-		ClearImmersivePanel();
-	}
-
 	const int32 PreviousMaximizedPanelIndex = MaximizedPanelIndex;
 	MaximizedPanelIndex = INDEX_NONE;
 	CurrentLayout = SavedLayout;
@@ -245,31 +239,29 @@ void SMediaProfileViewport::RestorePreviousLayout()
 	}
 }
 
-void SMediaProfileViewport::SetImmersivePanel(int32 InPanelIndex)
+void SMediaProfileViewport::SetImmersivePanel()
 {
 	TSharedPtr<SWindow> ParentWindow = FSlateApplication::Get().FindWidgetWindow(AsShared());
 	if (!ParentWindow.IsValid())
 	{
 		return;
 	}
-
-	if (!PanelContents.IsValidIndex(InPanelIndex))
-	{
-		return;
-	}
 	
-	ImmersivePanelIndex = InPanelIndex;
-	ParentWindow->SetFullWindowOverlayContent(PanelContents[InPanelIndex].Widget);
+	bImmersive = true;
+	ParentWindow->SetFullWindowOverlayContent(LayoutContainer);
 	
-	// Set keyboard focus on the panel widget after it has been expanded
-	// so that keyboard shortcuts still work for that panel without the user having to click on it again
-	FSlateApplication::Get().ClearKeyboardFocus();
-	FSlateApplication::Get().SetKeyboardFocus(PanelContents[InPanelIndex].Widget);
+	if (!PanelContents.IsEmpty())
+	{
+		// Reset keyboard focus back the panel widget after it has been repositioned in the panel layout
+		// so that keyboard shortcuts still work for that panel without the user having to click on it again
+		FSlateApplication::Get().ClearKeyboardFocus();
+		FSlateApplication::Get().SetKeyboardFocus(PanelContents[0].Widget);
+	}
 }
 
-bool SMediaProfileViewport::IsPanelImmersive(int32 InPanelIndex) const
+bool SMediaProfileViewport::IsPanelImmersive() const
 {
-	return ImmersivePanelIndex != INDEX_NONE && ImmersivePanelIndex == InPanelIndex;
+	return bImmersive;
 }
 
 void SMediaProfileViewport::ClearImmersivePanel()
@@ -280,17 +272,15 @@ void SMediaProfileViewport::ClearImmersivePanel()
 		return;
 	}
 
-	const int32 PreviousImmersivePanel = ImmersivePanelIndex;
-	
-	ImmersivePanelIndex = INDEX_NONE;
+	bImmersive = false;
 	ParentWindow->SetFullWindowOverlayContent(nullptr);
 
-	if (PanelContents.IsValidIndex(PreviousImmersivePanel))
+	if (!PanelContents.IsEmpty())
 	{
 		// Reset keyboard focus back the panel widget after it has been repositioned in the panel layout
 		// so that keyboard shortcuts still work for that panel without the user having to click on it again
 		FSlateApplication::Get().ClearKeyboardFocus();
-		FSlateApplication::Get().SetKeyboardFocus(PanelContents[PreviousImmersivePanel].Widget);
+		FSlateApplication::Get().SetKeyboardFocus(PanelContents[0].Widget);
 	}
 }
 
@@ -588,16 +578,13 @@ TSharedRef<SWidget> SMediaProfileViewport::ConstructViewportLayout()
 
 void SMediaProfileViewport::SetViewportLayout(EViewportLayout InViewportLayout, EViewportOrientation InViewportOrientation)
 {
-	if (ImmersivePanelIndex != INDEX_NONE)
-	{
-		ClearImmersivePanel();
-	}
-	
 	CurrentLayout = InViewportLayout;
 	CurrentOrientation = InViewportOrientation;
 	MaximizedPanelIndex = INDEX_NONE;
 	
 	UpdateDisplays();
+
+	FSlateApplication::Get().DismissAllMenus();
 }
 
 bool SMediaProfileViewport::IsViewportLayoutSet(EViewportLayout InViewportLayout, EViewportOrientation InViewportOrientation) const
diff --git a/Engine/Plugins/Media/MediaFrameworkUtilities/Source/MediaFrameworkUtilitiesEditor/Private/AssetEditor/SMediaProfileViewport.h b/Engine/Plugins/Media/MediaFrameworkUtilities/Source/MediaFrameworkUtilitiesEditor/Private/AssetEditor/SMediaProfileViewport.h
--- a/Engine/Plugins/Media/MediaFrameworkUtilities/Source/MediaFrameworkUtilitiesEditor/Private/AssetEditor/SMediaProfileViewport.h
+++ b/Engine/Plugins/Media/MediaFrameworkUtilities/Source/MediaFrameworkUtilitiesEditor/Private/AssetEditor/SMediaProfileViewport.h
@@ -81,11 +81,11 @@ public:
 	/** Restores the previous multi-panel layout from a maximized state, if one is saved */
 	void RestorePreviousLayout();
 
-	/** Makes the specified panel immersive */
-	void SetImmersivePanel(int32 InPanelIndex);
+	/** Sets the panel layout to display as immersive fullscreen */
+	void SetImmersivePanel();
 
-	/** Gets whether the specified panel is being displayed in immersive view */
-	bool IsPanelImmersive(int32 InPanelIndex) const;
+	/** Gets whether the panel layout is being displayed in immersive view */
+	bool IsPanelImmersive() const;
 	
 	/** Clears any immersive panel being displayed */
 	void ClearImmersivePanel();
@@ -149,6 +149,6 @@ private:
 	/** When maximizing from a multi-panel layout, this is the panel being maximized */
 	int32 MaximizedPanelIndex = INDEX_NONE;
 
-	/** When toggling immersive mode, this is the panel being displayed */
-	int32 ImmersivePanelIndex = INDEX_NONE;
+	/** Whether the panel layout is being displayed in immersive fullscreen */
+	bool bImmersive = false;
 };
\ No newline at end of file
