Media Profile: Changed active viewport media capture to allow input

#Synced-CL 51425529

diff --git a/Engine/Plugins/Media/MediaFrameworkUtilities/Source/MediaFrameworkUtilities/Private/Profile/MediaProfilePlaybackManager.cpp b/Engine/Plugins/Media/MediaFrameworkUtilities/Source/MediaFrameworkUtilities/Private/Profile/MediaProfilePlaybackManager.cpp
--- a/Engine/Plugins/Media/MediaFrameworkUtilities/Source/MediaFrameworkUtilities/Private/Profile/MediaProfilePlaybackManager.cpp
+++ b/Engine/Plugins/Media/MediaFrameworkUtilities/Source/MediaFrameworkUtilities/Private/Profile/MediaProfilePlaybackManager.cpp
@@ -235,7 +235,6 @@ namespace MediaProfilePlaybackManager
 					CachedViewportFlags.bRealTime = ViewportClient.IsRealtime();
 					CachedViewportFlags.bSetListenerPosition = ViewportClient.bSetListenerPosition;
 					CachedViewportFlags.bDrawAxes = ViewportClient.bDrawAxes;
-					CachedViewportFlags.bDisableInput = ViewportClient.bDisableInput;
 					CachedViewportFlags.bAllowCinematicControl = ViewportClient.AllowsCinematicControl();
 					CachedViewportFlags.bDrawAlphaBlendedCheckerboard = ViewportClient.ChannelMaskParams.bDrawAlphaBlendedCheckerboard;
 					CachedViewportFlags.VisibilityDelegate = ViewportClient.VisibilityDelegate;
@@ -250,9 +249,6 @@ namespace MediaProfilePlaybackManager
 					FLevelEditorModule& LevelEditorModule = FModuleManager::GetModuleChecked<FLevelEditorModule>(TEXT("LevelEditor"));
 					LevelEditorModule.OnGetVisualizeCheckerboardEnabled().BindLambda([] { return false; });
 
-					const bool bDisableInput = InCaptureOptions.ResizeMethod != EMediaCaptureResizeMethod::ResizeInRenderPass;
-					ViewportClient.bDisableInput = bDisableInput ? true : ViewportClient.bDisableInput;
-
 					if (InCaptureOptions.ResizeMethod == EMediaCaptureResizeMethod::ResizeSource)
 					{
 						if (TSharedPtr<FSceneViewport> LevelViewport = ActiveViewportInterface->GetSharedActiveViewport())
@@ -269,6 +265,7 @@ namespace MediaProfilePlaybackManager
 							}
 						
 							LevelViewport->SetFixedViewportSize(TargetSize.X, TargetSize.Y);
+							LevelViewport->InvalidateHitProxy();
 						}
 					}
 				}
@@ -287,7 +284,6 @@ namespace MediaProfilePlaybackManager
 				ViewportClient.SetRealtime(CachedViewportFlags.bRealTime);
 				ViewportClient.bSetListenerPosition = CachedViewportFlags.bSetListenerPosition;
 				ViewportClient.bDrawAxes = CachedViewportFlags.bDrawAxes;
-				ViewportClient.bDisableInput = CachedViewportFlags.bDisableInput;
 				ViewportClient.SetAllowCinematicControl(CachedViewportFlags.bAllowCinematicControl);
 				ViewportClient.ChannelMaskParams.bDrawAlphaBlendedCheckerboard = CachedViewportFlags.bDrawAlphaBlendedCheckerboard;
 				ViewportClient.VisibilityDelegate = CachedViewportFlags.VisibilityDelegate;
@@ -298,6 +294,7 @@ namespace MediaProfilePlaybackManager
 				if (TSharedPtr<FSceneViewport> LevelViewport = PinnedViewportInterface->GetSharedActiveViewport())
 				{
 					LevelViewport->SetFixedViewportSize(CachedViewportFlags.ViewportSize.X, CachedViewportFlags.ViewportSize.Y);
+					LevelViewport->InvalidateHitProxy();
 				}
 			}
 #endif
@@ -337,7 +334,6 @@ namespace MediaProfilePlaybackManager
 			bool bRealTime = false;
 			bool bSetListenerPosition = false;
 			bool bDrawAxes = false;
-			bool bDisableInput = false;
 			bool bAllowCinematicControl = false;
 			bool bDrawAlphaBlendedCheckerboard = false;
 			FEngineShowFlags EngineShowFlags = ESFIM_Editor;
