Media Profile: Make TC toolbar entry a clickable button that opens/focuses the TC/Genlock panel

#Synced-CL 51425529

diff --git a/Engine/Plugins/Media/MediaFrameworkUtilities/Source/MediaFrameworkUtilitiesEditor/Private/AssetEditor/MediaProfileEditor.cpp b/Engine/Plugins/Media/MediaFrameworkUtilities/Source/MediaFrameworkUtilitiesEditor/Private/AssetEditor/MediaProfileEditor.cpp
--- a/Engine/Plugins/Media/MediaFrameworkUtilities/Source/MediaFrameworkUtilitiesEditor/Private/AssetEditor/MediaProfileEditor.cpp
+++ b/Engine/Plugins/Media/MediaFrameworkUtilities/Source/MediaFrameworkUtilitiesEditor/Private/AssetEditor/MediaProfileEditor.cpp
@@ -337,6 +337,11 @@ TSharedRef<SDockTab> FMediaProfileEditor::SpawnTab_Timecode(const FSpawnTabArgs&
 		];
 }
 
+void FMediaProfileEditor::OpenTab(FName InTabId)
+{
+	TabManager->TryInvokeTab(InTabId);
+}
+
 void FMediaProfileEditor::ExtendToolbar()
 {
 	TSharedPtr<FExtender> ToolbarExtender = MakeShared<FExtender>();
@@ -347,7 +352,38 @@ void FMediaProfileEditor::ExtendToolbar()
 		GetToolkitCommands(),
 		FToolBarExtensionDelegate::CreateLambda([this](FToolBarBuilder& ToolbarBuilder)
 		{
-			ToolbarBuilder.AddWidget(SAssignNew(TimecodeToolbarEntry, SMediaFrameworkTimecodeGenlockHeader));
+			ToolbarBuilder.AddWidget(
+				SAssignNew(TimecodeToolbarEntry, SMediaFrameworkTimecodeGenlockHeader)
+				.Visibility_Lambda([]()
+				{
+					if (const UMediaProfileEditorUserSettings* UserSettings = GetDefault<UMediaProfileEditorUserSettings>())
+					{
+						return (UserSettings->bShowTimecodeInEditorToolbar || UserSettings->bShowGenlockInEditorToolbar) ? EVisibility::Visible : EVisibility::Collapsed;
+					}
+
+					return EVisibility::Visible;
+				})
+				.ShowTimecode_Lambda([]()
+				{
+					if (const UMediaProfileEditorUserSettings* UserSettings = GetDefault<UMediaProfileEditorUserSettings>())
+					{
+						return UserSettings->bShowTimecodeInEditorToolbar;
+					}
+
+					return true;
+				})
+				.ShowGenlock_Lambda([]()
+				{
+					if (const UMediaProfileEditorUserSettings* UserSettings = GetDefault<UMediaProfileEditorUserSettings>())
+					{
+						return UserSettings->bShowGenlockInEditorToolbar;
+					}
+
+					return true;
+				})
+				.IsButton(true)
+				.OnOpenTimecodeTab(this, &FMediaProfileEditor::OpenTab, TimecodeTabId)
+			);
 		})
 	);
 
diff --git a/Engine/Plugins/Media/MediaFrameworkUtilities/Source/MediaFrameworkUtilitiesEditor/Private/AssetEditor/MediaProfileEditor.h b/Engine/Plugins/Media/MediaFrameworkUtilities/Source/MediaFrameworkUtilitiesEditor/Private/AssetEditor/MediaProfileEditor.h
--- a/Engine/Plugins/Media/MediaFrameworkUtilities/Source/MediaFrameworkUtilitiesEditor/Private/AssetEditor/MediaProfileEditor.h
+++ b/Engine/Plugins/Media/MediaFrameworkUtilities/Source/MediaFrameworkUtilitiesEditor/Private/AssetEditor/MediaProfileEditor.h
@@ -77,6 +77,9 @@ private:
 	TSharedRef<SDockTab> SpawnTab_Details(const FSpawnTabArgs& InArgs);
 	TSharedRef<SDockTab> SpawnTab_Timecode(const FSpawnTabArgs& InArgs);
 
+	/** Raised when something has requested that a tab be manually opened/focused */
+	void OpenTab(FName InTabId);
+	
 	void ExtendToolbar();
 
 	/** Refreshes the selected media items, reopening media sources and outputs */
diff --git a/Engine/Plugins/Media/MediaFrameworkUtilities/Source/MediaFrameworkUtilitiesEditor/Private/AssetEditor/MediaProfileEditorUserSettings.h b/Engine/Plugins/Media/MediaFrameworkUtilities/Source/MediaFrameworkUtilitiesEditor/Private/AssetEditor/MediaProfileEditorUserSettings.h
--- a/Engine/Plugins/Media/MediaFrameworkUtilities/Source/MediaFrameworkUtilitiesEditor/Private/AssetEditor/MediaProfileEditorUserSettings.h
+++ b/Engine/Plugins/Media/MediaFrameworkUtilities/Source/MediaFrameworkUtilitiesEditor/Private/AssetEditor/MediaProfileEditorUserSettings.h
@@ -38,10 +38,18 @@ public:
 	UPROPERTY(Config)
 	bool bShowTimecodeInViewportToolbar = true;
 
+	/** Whether to show the timecode entry in the main editor toolbar */
+	UPROPERTY(Config)
+	bool bShowTimecodeInEditorToolbar = true;
+	
 	/** Whether to show the genlock entry in the viewport toolbar */
 	UPROPERTY(Config)
 	bool bShowGenlockInViewportToolbar = true;
 
+	/** Whether to show the genlock entry in the main editor toolbar */
+	UPROPERTY(Config)
+	bool bShowGenlockInEditorToolbar = true;
+	
 	/** Whether to show the Channels menu in the viewports */
 	UPROPERTY(Config)
 	bool bShowChannelsInViewportToolbar = true;
diff --git a/Engine/Plugins/Media/MediaFrameworkUtilities/Source/MediaFrameworkUtilitiesEditor/Private/UI/SMediaFrameworkTimecodeGenlockHeader.cpp b/Engine/Plugins/Media/MediaFrameworkUtilities/Source/MediaFrameworkUtilitiesEditor/Private/UI/SMediaFrameworkTimecodeGenlockHeader.cpp
--- a/Engine/Plugins/Media/MediaFrameworkUtilities/Source/MediaFrameworkUtilitiesEditor/Private/UI/SMediaFrameworkTimecodeGenlockHeader.cpp
+++ b/Engine/Plugins/Media/MediaFrameworkUtilities/Source/MediaFrameworkUtilitiesEditor/Private/UI/SMediaFrameworkTimecodeGenlockHeader.cpp
@@ -9,6 +9,7 @@
 #include "Profile/IMediaProfileManager.h"
 #include "Profile/MediaProfile.h"
 #include "Widgets/Images/SImage.h"
+#include "Widgets/Input/SButton.h"
 #include "Widgets/Layout/SSeparator.h"
 #include "Widgets/Text/STextBlock.h"
 
@@ -16,9 +17,11 @@
 
 void SMediaFrameworkTimecodeGenlockHeader::Construct(const FArguments& InArgs)
 {
-	ChildSlot
-	[
-		SNew(SHorizontalBox)
+	ShowTimecodeAttr = InArgs._ShowTimecode;
+	ShowGenlockAttr = InArgs._ShowGenlock;
+	OnOpenTimecodeTab = InArgs._OnOpenTimecodeTab;
+
+	TSharedRef<SWidget> ButtonContent = SNew(SHorizontalBox)
 		.ToolTipText(this, &SMediaFrameworkTimecodeGenlockHeader::GetTooltipText)
 
 		+SHorizontalBox::Slot()
@@ -30,6 +33,7 @@ void SMediaFrameworkTimecodeGenlockHeader::Construct(const FArguments& InArgs)
 			SNew(STextBlock)
 			.Text(this, &SMediaFrameworkTimecodeGenlockHeader::GetTimecodeText)
 			.Font(FCoreStyle::GetDefaultFontStyle("Bold", 10))
+			.Visibility(this, &SMediaFrameworkTimecodeGenlockHeader::GetElementVisibility, true, false)
 		]
 		
 		+SHorizontalBox::Slot()
@@ -38,6 +42,7 @@ void SMediaFrameworkTimecodeGenlockHeader::Construct(const FArguments& InArgs)
 			SNew(SSeparator)
 			.Orientation(EOrientation::Orient_Vertical)
 			.Thickness(2)
+			.Visibility(this, &SMediaFrameworkTimecodeGenlockHeader::GetElementVisibility, true, true)
 		]
 
 		+SHorizontalBox::Slot()
@@ -49,6 +54,7 @@ void SMediaFrameworkTimecodeGenlockHeader::Construct(const FArguments& InArgs)
 			SNew(SImage)
 			.Image(this, &SMediaFrameworkTimecodeGenlockHeader::GetGenlockIcon)
 			.DesiredSizeOverride(FVector2D(20, 20))
+			.Visibility(this, &SMediaFrameworkTimecodeGenlockHeader::GetElementVisibility, false, true)
 		]
 		
 		+SHorizontalBox::Slot()
@@ -60,7 +66,24 @@ void SMediaFrameworkTimecodeGenlockHeader::Construct(const FArguments& InArgs)
 			SNew(STextBlock)
 			.Text(this, &SMediaFrameworkTimecodeGenlockHeader::GetGenlockText)
 			.Font(FCoreStyle::GetDefaultFontStyle("Bold", 10))
-		]
+			.Visibility(this, &SMediaFrameworkTimecodeGenlockHeader::GetElementVisibility, false, true)
+		];
+
+	TSharedRef<SWidget> ChildContent = ButtonContent;
+	if (InArgs._IsButton)
+	{
+		ChildContent = SNew(SButton)
+			.ButtonStyle(FAppStyle::Get(), "HoverHintOnly")
+			.ContentPadding(0.0f)
+			.OnClicked(this, &SMediaFrameworkTimecodeGenlockHeader::OnButtonClicked)
+			[
+				ButtonContent
+			];
+	}
+	
+	ChildSlot
+	[
+		ChildContent
 	];
 }
 
@@ -137,4 +160,31 @@ FText SMediaFrameworkTimecodeGenlockHeader::GetGenlockText() const
 	return FText::FromString(CustomTimeStep->GetDisplayName());
 }
 
+FReply SMediaFrameworkTimecodeGenlockHeader::OnButtonClicked()
+{
+	OnOpenTimecodeTab.ExecuteIfBound();
+	return FReply::Handled();
+}
+
+EVisibility SMediaFrameworkTimecodeGenlockHeader::GetElementVisibility(bool bTimecode, bool bGenlock) const
+{
+	const bool bShowTimecode = ShowTimecodeAttr.Get(true);
+	const bool bShowGenlock = ShowGenlockAttr.Get(true);
+	
+	if (bTimecode && bGenlock)
+	{
+		return (bShowTimecode && bShowGenlock) ? EVisibility::Visible : EVisibility::Collapsed;
+	}
+	if (bTimecode)
+	{
+		return bShowTimecode ? EVisibility::Visible : EVisibility::Collapsed;
+	}
+	if (bGenlock)
+	{
+		return bShowGenlock ? EVisibility::Visible : EVisibility::Collapsed;
+	}
+	
+	return EVisibility::Collapsed;
+}
+
 #undef LOCTEXT_NAMESPACE
\ No newline at end of file
diff --git a/Engine/Plugins/Media/MediaFrameworkUtilities/Source/MediaFrameworkUtilitiesEditor/Private/UI/SMediaFrameworkTimecodeGenlockHeader.h b/Engine/Plugins/Media/MediaFrameworkUtilities/Source/MediaFrameworkUtilitiesEditor/Private/UI/SMediaFrameworkTimecodeGenlockHeader.h
--- a/Engine/Plugins/Media/MediaFrameworkUtilities/Source/MediaFrameworkUtilitiesEditor/Private/UI/SMediaFrameworkTimecodeGenlockHeader.h
+++ b/Engine/Plugins/Media/MediaFrameworkUtilities/Source/MediaFrameworkUtilitiesEditor/Private/UI/SMediaFrameworkTimecodeGenlockHeader.h
@@ -12,7 +12,15 @@ class UMediaProfile;
 class SMediaFrameworkTimecodeGenlockHeader : public SCompoundWidget
 {
 public:
-	SLATE_BEGIN_ARGS(SMediaFrameworkTimecodeGenlockHeader) { }
+	SLATE_BEGIN_ARGS(SMediaFrameworkTimecodeGenlockHeader)
+		: _IsButton(false)
+		, _ShowTimecode(true)
+		, _ShowGenlock(true)
+	{ }
+		SLATE_ARGUMENT(bool, IsButton)
+		SLATE_ATTRIBUTE(bool, ShowTimecode)
+		SLATE_ATTRIBUTE(bool, ShowGenlock)
+		SLATE_EVENT(FSimpleDelegate, OnOpenTimecodeTab)
 	SLATE_END_ARGS()
 
 	/** Constructs this widget with InArgs */
@@ -30,4 +38,20 @@ private:
 
 	/** Gets the text to display for the genlock */
 	FText GetGenlockText() const;
+
+	/** Raised when the toolbar button is clicked */
+	FReply OnButtonClicked();
+
+	/** Gets whether certain elements of the toolbar entry should be visible */
+	EVisibility GetElementVisibility(bool bTimecode, bool bGenlock) const;
+	
+private:
+	/** Gets whether the timecode should be shown */
+	TAttribute<bool> ShowTimecodeAttr;
+
+	/** Gets whether the genlock should be shown */
+	TAttribute<bool> ShowGenlockAttr;
+	
+	/** Raised when the header is requesting the TC/Genlock tab be opened/focused */
+	FSimpleDelegate OnOpenTimecodeTab;
 };
diff --git a/Engine/Plugins/Media/MediaFrameworkUtilities/Source/MediaFrameworkUtilitiesEditor/Private/UI/SMediaFrameworkTimecodeGenlockPanel.cpp b/Engine/Plugins/Media/MediaFrameworkUtilities/Source/MediaFrameworkUtilitiesEditor/Private/UI/SMediaFrameworkTimecodeGenlockPanel.cpp
--- a/Engine/Plugins/Media/MediaFrameworkUtilities/Source/MediaFrameworkUtilitiesEditor/Private/UI/SMediaFrameworkTimecodeGenlockPanel.cpp
+++ b/Engine/Plugins/Media/MediaFrameworkUtilities/Source/MediaFrameworkUtilitiesEditor/Private/UI/SMediaFrameworkTimecodeGenlockPanel.cpp
@@ -2,15 +2,19 @@
 
 #include "SMediaFrameworkTimecodeGenlockPanel.h"
 
+#include "DetailCategoryBuilder.h"
 #include "DetailLayoutBuilder.h"
+#include "DetailWidgetRow.h"
 #include "IDetailCustomization.h"
 #include "SMediaFrameworkTimecodeGenlockHeader.h"
+#include "AssetEditor/MediaProfileEditorUserSettings.h"
 #include "Engine/EngineCustomTimeStep.h"
 #include "Engine/TimecodeProvider.h"
 #include "Misc/Timecode.h"
 #include "Modules/ModuleManager.h"
 #include "Profile/MediaProfile.h"
 #include "Widgets/Images/SImage.h"
+#include "Widgets/Input/SCheckBox.h"
 #include "Widgets/Layout/SSeparator.h"
 #include "Widgets/Text/STextBlock.h"
 
@@ -23,6 +27,66 @@ class FMediaProfileTimecodeGenlockDetailsCustomization : public IDetailCustomiza
 	{
 		DetailBuilder.HideCategory(TEXT("Inputs"));
 		DetailBuilder.HideCategory(TEXT("Outputs"));
+
+		IDetailCategoryBuilder& TimecodeCategory = DetailBuilder.EditCategory(TEXT("Timecode Provider"));
+
+		TimecodeCategory.AddCustomRow(LOCTEXT("ShowTimecodeFilterString", "Show Timecode"))
+		.NameContent()
+		[
+			SNew(STextBlock)
+			.Text(LOCTEXT("ShowTimecodeLabel", "Show Timecode in Toolbar"))
+			.Font(IDetailLayoutBuilder::GetDetailFont())
+		]
+		.ValueContent()
+		[
+			SNew(SCheckBox)
+			.IsChecked_Lambda([]()
+			{
+				if (const UMediaProfileEditorUserSettings* UserSettings = GetDefault<UMediaProfileEditorUserSettings>())
+				{
+					return UserSettings->bShowTimecodeInEditorToolbar ? ECheckBoxState::Checked : ECheckBoxState::Unchecked;
+				}
+
+				return ECheckBoxState::Checked;
+			})
+			.OnCheckStateChanged_Lambda([](ECheckBoxState CheckBoxState)
+			{
+				if (UMediaProfileEditorUserSettings* UserSettings = GetMutableDefault<UMediaProfileEditorUserSettings>())
+				{
+					UserSettings->bShowTimecodeInEditorToolbar = CheckBoxState == ECheckBoxState::Checked;
+				}
+			})
+		];
+
+		IDetailCategoryBuilder& GenlockCategory = DetailBuilder.EditCategory(TEXT("Genlock"));
+
+		GenlockCategory.AddCustomRow(LOCTEXT("ShowGenlockFilterString", "Show Genlock"))
+		.NameContent()
+		[
+			SNew(STextBlock)
+			.Text(LOCTEXT("ShowGenlockLabel", "Show Genlock in Toolbar"))
+			.Font(IDetailLayoutBuilder::GetDetailFont())
+		]
+		.ValueContent()
+		[
+			SNew(SCheckBox)
+			.IsChecked_Lambda([]()
+			{
+				if (const UMediaProfileEditorUserSettings* UserSettings = GetDefault<UMediaProfileEditorUserSettings>())
+				{
+					return UserSettings->bShowGenlockInEditorToolbar ? ECheckBoxState::Checked : ECheckBoxState::Unchecked;
+				}
+
+				return ECheckBoxState::Checked;
+			})
+			.OnCheckStateChanged_Lambda([](ECheckBoxState CheckBoxState)
+			{
+				if (UMediaProfileEditorUserSettings* UserSettings = GetMutableDefault<UMediaProfileEditorUserSettings>())
+				{
+					UserSettings->bShowGenlockInEditorToolbar = CheckBoxState == ECheckBoxState::Checked;
+				}
+			})
+		];
 	}
 };
 
