Media Profile: Added new media capture method that outputs the contents of a media texture.

#Synced-CL 51697894

diff --git a/Engine/Plugins/Media/AjaMedia/Source/AjaMediaOutput/Private/AjaMediaCapture.cpp b/Engine/Plugins/Media/AjaMedia/Source/AjaMediaOutput/Private/AjaMediaCapture.cpp
--- a/Engine/Plugins/Media/AjaMedia/Source/AjaMediaOutput/Private/AjaMediaCapture.cpp
+++ b/Engine/Plugins/Media/AjaMedia/Source/AjaMediaOutput/Private/AjaMediaCapture.cpp
@@ -322,6 +322,12 @@ bool UAjaMediaCapture::UpdateRenderTargetImpl(UTextureRenderTarget2D* InRenderTa
 	return true;
 }
 
+bool UAjaMediaCapture::UpdateMediaTextureImpl(UMediaTexture* InMediaTexture)
+{
+	RestoreViewportTextureAlpha(GetCapturingSceneViewport());
+	return true;
+}
+
 bool UAjaMediaCapture::UpdateAudioDeviceImpl(const FAudioDeviceHandle& InAudioDeviceHandle)
 {
 	return CreateAudioOutput(InAudioDeviceHandle, Cast<UAjaMediaOutput>(MediaOutput));
diff --git a/Engine/Plugins/Media/AjaMedia/Source/AjaMediaOutput/Public/AjaMediaCapture.h b/Engine/Plugins/Media/AjaMedia/Source/AjaMediaOutput/Public/AjaMediaCapture.h
--- a/Engine/Plugins/Media/AjaMedia/Source/AjaMediaOutput/Public/AjaMediaCapture.h
+++ b/Engine/Plugins/Media/AjaMedia/Source/AjaMediaOutput/Public/AjaMediaCapture.h
@@ -46,6 +46,7 @@ protected:
 	virtual bool PostInitializeCaptureViewport(TSharedPtr<FSceneViewport>& InSceneViewport) override;
 	virtual bool UpdateSceneViewportImpl(TSharedPtr<FSceneViewport>& InSceneViewport) override;
 	virtual bool UpdateRenderTargetImpl(UTextureRenderTarget2D* InRenderTarget) override;
+	virtual bool UpdateMediaTextureImpl(UMediaTexture* InMediaTexture) override;
 	virtual bool UpdateAudioDeviceImpl(const FAudioDeviceHandle& InAudioDeviceHandle) override;
 	virtual void StopCaptureImpl(bool bAllowPendingFrameToBeProcess) override;
 	virtual void PostCaptureCleanupImpl() override;
diff --git a/Engine/Plugins/Media/MediaFrameworkUtilities/Source/MediaFrameworkUtilities/Private/Profile/MediaProfilePlaybackManager.cpp b/Engine/Plugins/Media/MediaFrameworkUtilities/Source/MediaFrameworkUtilities/Private/Profile/MediaProfilePlaybackManager.cpp
--- a/Engine/Plugins/Media/MediaFrameworkUtilities/Source/MediaFrameworkUtilities/Private/Profile/MediaProfilePlaybackManager.cpp
+++ b/Engine/Plugins/Media/MediaFrameworkUtilities/Source/MediaFrameworkUtilities/Private/Profile/MediaProfilePlaybackManager.cpp
@@ -997,6 +997,64 @@ UMediaCapture* UMediaProfilePlaybackManager::OpenRenderTargetOutput(UMediaOutput
 	return OpenRenderTargetOutputFromIndex(MediaOutputIndex, InRenderTarget, InCaptureOptions, bAutoRestartCapture);
 }
 
+UMediaCapture* UMediaProfilePlaybackManager::OpenMediaTextureOutputFromIndex(int32 InMediaOutputIndex, UMediaTexture* InMediaTexture, const FMediaCaptureOptions& InCaptureOptions, bool bAutoRestartCapture)
+{
+	UMediaProfile* MediaProfile = GetOwningMediaProfile();
+	if (!MediaProfile)
+	{
+		return nullptr;
+	}
+
+	UMediaOutput* MediaOutput = MediaProfile->GetMediaOutput(InMediaOutputIndex);
+	if (!MediaOutput)
+	{
+		return nullptr;
+	}
+	
+	UMediaCapture* MediaCapture = FindOrCreateMediaCaptureForIndex(InMediaOutputIndex);
+	if (!MediaCapture)
+	{
+		return nullptr;
+	}
+	
+	// Check to see that the media capture is not already capturing
+	EMediaCaptureState CaptureState = MediaCapture->GetState();
+	if (CaptureState != EMediaCaptureState::Stopped && CaptureState != EMediaCaptureState::StopRequested && CaptureState != EMediaCaptureState::Error)
+	{
+		return MediaCapture;
+	};
+
+	if (MediaCapture->CaptureMediaTexture(InMediaTexture, InCaptureOptions))
+	{
+#if WITH_EDITOR
+		if (bAutoRestartCapture)
+		{
+			MediaOutput->OnOutputModified().AddUObject(this, &UMediaProfilePlaybackManager::RestartMediaTextureOutput, InMediaTexture, InCaptureOptions, bAutoRestartCapture);
+		}
+#endif
+		
+		return MediaCapture;
+	}
+
+	// If the capture failed to start, mark the capture as errored, clean up its capture resources, and return it
+	CapturesWithError.Add(MediaCapture);
+	CleanUpCaptureResources(MediaOutput);
+	
+	return MediaCapture;
+}
+
+UMediaCapture* UMediaProfilePlaybackManager::OpenMediaTextureOutput(UMediaOutput* InMediaOutput, UMediaTexture* InMediaTexture, const FMediaCaptureOptions& InCaptureOptions, bool bAutoRestartCapture)
+{
+	UMediaProfile* MediaProfile = GetOwningMediaProfile();
+	if (!MediaProfile)
+	{
+		return nullptr;
+	}
+	
+	int32 MediaOutputIndex = MediaProfile->FindMediaOutputIndex(InMediaOutput);
+	return OpenMediaTextureOutputFromIndex(MediaOutputIndex, InMediaTexture, InCaptureOptions, bAutoRestartCapture);
+}
+
 void UMediaProfilePlaybackManager::RestartActiveViewportOutput(UMediaOutput* InMediaOutput, FMediaCaptureOptions InCaptureOptions, bool bAutoRestartCapture)
 {
 	TWeakObjectPtr<UMediaOutput> WeakMediaOutput = InMediaOutput;
@@ -1044,6 +1102,23 @@ void UMediaProfilePlaybackManager::RestartRenderTargetOutput(UMediaOutput* InMed
 	QueueRestartCapture(InMediaOutput, Callback);
 }
 
+void UMediaProfilePlaybackManager::RestartMediaTextureOutput(UMediaOutput* InMediaOutput, UMediaTexture* InMediaTexture, FMediaCaptureOptions InCaptureOptions, bool bAutoRestartCapture)
+{
+	TWeakObjectPtr<UMediaTexture> WeakMediaTexture = InMediaTexture;
+	TWeakObjectPtr<UMediaOutput> WeakMediaOutput = InMediaOutput;
+	auto Callback = [WeakMediaTexture, WeakMediaOutput, InCaptureOptions, bAutoRestartCapture](UMediaProfilePlaybackManager* Manager)
+	{
+		TStrongObjectPtr<UMediaTexture> MediaTexture = WeakMediaTexture.Pin();
+		TStrongObjectPtr<UMediaOutput> MediaOutput = WeakMediaOutput.Pin();
+
+		if (Manager && MediaTexture && MediaOutput)
+		{
+			Manager->OpenMediaTextureOutput(MediaOutput.Get(), MediaTexture.Get(), InCaptureOptions, bAutoRestartCapture);
+		}
+	};
+	QueueRestartCapture(InMediaOutput, Callback);
+}
+
 bool UMediaProfilePlaybackManager::IsOutputCapturingFromIndex(int32 InMediaOutputIndex) const
 {
 	UMediaProfile* MediaProfile = GetOwningMediaProfile();
diff --git a/Engine/Plugins/Media/MediaFrameworkUtilities/Source/MediaFrameworkUtilities/Public/Profile/MediaProfilePlaybackManager.h b/Engine/Plugins/Media/MediaFrameworkUtilities/Source/MediaFrameworkUtilities/Public/Profile/MediaProfilePlaybackManager.h
--- a/Engine/Plugins/Media/MediaFrameworkUtilities/Source/MediaFrameworkUtilities/Public/Profile/MediaProfilePlaybackManager.h
+++ b/Engine/Plugins/Media/MediaFrameworkUtilities/Source/MediaFrameworkUtilities/Public/Profile/MediaProfilePlaybackManager.h
@@ -234,6 +234,26 @@ public:
 	 */
 	UMediaCapture* OpenRenderTargetOutput(UMediaOutput* InMediaOutput, UTextureRenderTarget2D* InRenderTarget, const FMediaCaptureOptions& InCaptureOptions, bool bAutoRestartCapture = false);
 
+	/**
+	 * Starts a capture of a media texture for a specified media output
+	 * @param InMediaOutputIndex The index of the media output in the media profile
+	 * @param InMediaTexture The media texture to capture
+	 * @param InCaptureOptions The capture options to capture with
+	 * @param bAutoRestartCapture Indicates that this capture will automatically restart if any modifications to the media output are detected
+	 * @return The media capture object, if the capture was started successfully; otherwise, null
+	 */
+	UMediaCapture* OpenMediaTextureOutputFromIndex(int32 InMediaOutputIndex, UMediaTexture* InMediaTexture, const FMediaCaptureOptions& InCaptureOptions, bool bAutoRestartCapture = false);
+	
+	/**
+	 * Starts a capture of a media texture for a specified media output
+	 * @param InMediaOutput The media output in the media profile
+	 * @param InMediaTexture The media texture to capture
+	 * @param InCaptureOptions The capture options to capture with
+	 * @param bAutoRestartCapture Indicates that this capture will automatically restart if any modifications to the media output are detected
+	 * @return The media capture object, if the capture was started successfully; otherwise, null
+	 */
+	UMediaCapture* OpenMediaTextureOutput(UMediaOutput* InMediaOutput, UMediaTexture* InMediaTexture, const FMediaCaptureOptions& InCaptureOptions, bool bAutoRestartCapture = false);
+	
 	 /* 
 	 * Restarts a capture of an existing engine viewport for a specified media output
 	 * @param InMediaOutput The media output in the media profile
@@ -259,6 +279,15 @@ public:
 	 */
 	void RestartRenderTargetOutput(UMediaOutput* InMediaOutput, UTextureRenderTarget2D* InRenderTarget, FMediaCaptureOptions InCaptureOptions, bool bAutoRestartCapture = false);
 
+	/**
+	 * Restarts a capture of a media texture for a specified media output
+	 * @param InMediaOutput The media output in the media profile
+	 * @param InMediaTexture The media texture to capture
+	 * @param InCaptureOptions The capture options to capture with
+	 * @param bAutoRestartCapture Indicates that this capture will automatically restart if any modifications to the media output are detected
+	 */
+	void RestartMediaTextureOutput(UMediaOutput* InMediaOutput, UMediaTexture* InMediaTexture, FMediaCaptureOptions InCaptureOptions, bool bAutoRestartCapture = false);
+	
 	/**
 	 * Gets whether there is an active capture for the specified media output index
 	 * @param InMediaOutputIndex The index of the media output within the media profile to check
diff --git a/Engine/Plugins/Media/MediaFrameworkUtilities/Source/MediaFrameworkUtilitiesEditor/MediaFrameworkUtilitiesEditor.Build.cs b/Engine/Plugins/Media/MediaFrameworkUtilities/Source/MediaFrameworkUtilitiesEditor/MediaFrameworkUtilitiesEditor.Build.cs
--- a/Engine/Plugins/Media/MediaFrameworkUtilities/Source/MediaFrameworkUtilitiesEditor/MediaFrameworkUtilitiesEditor.Build.cs
+++ b/Engine/Plugins/Media/MediaFrameworkUtilities/Source/MediaFrameworkUtilitiesEditor/MediaFrameworkUtilitiesEditor.Build.cs
@@ -6,6 +6,12 @@ namespace UnrealBuildTool.Rules
 	{
 		public MediaFrameworkUtilitiesEditor(ReadOnlyTargetRules Target) : base(Target)
 		{
+			PublicDependencyModuleNames.AddRange(
+				new string[]
+				{
+					"PropertyEditor"
+				});
+
 			PrivateDependencyModuleNames.AddRange(
 				new string[]
 				{
diff --git a/Engine/Plugins/Media/MediaFrameworkUtilities/Source/MediaFrameworkUtilitiesEditor/Private/AssetEditor/MediaProfileEditor.cpp b/Engine/Plugins/Media/MediaFrameworkUtilities/Source/MediaFrameworkUtilitiesEditor/Private/AssetEditor/MediaProfileEditor.cpp
--- a/Engine/Plugins/Media/MediaFrameworkUtilities/Source/MediaFrameworkUtilitiesEditor/Private/AssetEditor/MediaProfileEditor.cpp
+++ b/Engine/Plugins/Media/MediaFrameworkUtilities/Source/MediaFrameworkUtilitiesEditor/Private/AssetEditor/MediaProfileEditor.cpp
@@ -248,45 +248,48 @@ bool FMediaProfileEditor::CanMediaOutputCapture(UMediaOutput* InMediaOutput) con
 
 	// Return true if there is at least one valid capture configuration for the specified media output
 	bool bHasValidCaptureSettings = false;
-	if (CaptureSettings->CurrentViewportMediaOutput.MediaOutput == InMediaOutput)
-	{
-		//Check to see if there is a valid editor viewport that can be captured from
-		TSharedPtr<FSceneViewport> Viewport = MediaProfileBeingEdited->GetPlaybackManager()->GetActiveViewport(InMediaOutput);
-		bHasValidCaptureSettings = Viewport.IsValid();
-	}
 	
-	for (const FMediaFrameworkCaptureCameraViewportCameraOutputInfo& ViewportOutputInfo : CaptureSettings->ViewportCaptures)
-	{
-		if (ViewportOutputInfo.MediaOutput != InMediaOutput)
-		{
-			continue;
-		}
-
-		const bool bHasValidCamera = ViewportOutputInfo.Cameras.ContainsByPredicate([](const TSoftObjectPtr<AActor>& ActorRef)
+	using namespace UE::MediaFrameworkWorldSettings::Helpers;
+	
+	ForEachOutputInfoForMediaOutput<FMediaFrameworkCaptureCurrentViewportOutputInfo>(CaptureSettings, InMediaOutput, 
+		[this, &bHasValidCaptureSettings](const FMediaFrameworkCaptureCurrentViewportOutputInfo& OutputInfo)
 		{
-			return ActorRef.IsValid();
+			//Check to see if there is a valid editor viewport that can be captured from
+			TSharedPtr<FSceneViewport> Viewport = MediaProfileBeingEdited->GetPlaybackManager()->GetActiveViewport(OutputInfo.MediaOutput);
+			bHasValidCaptureSettings = Viewport.IsValid();
 		});
-		
-		if (bHasValidCamera)
+	
+	ForEachOutputInfoForMediaOutput<FMediaFrameworkCaptureCameraViewportCameraOutputInfo>(CaptureSettings, InMediaOutput, 
+		[this, &bHasValidCaptureSettings](const FMediaFrameworkCaptureCameraViewportCameraOutputInfo& OutputInfo)
 		{
-			bHasValidCaptureSettings = true;
-			break;
-		}
-	}
+			const bool bHasValidCamera = OutputInfo.Cameras.ContainsByPredicate([](const TSoftObjectPtr<AActor>& ActorRef)
+			{
+				return ActorRef.IsValid();
+			});
+			
+			if (bHasValidCamera)
+			{
+				bHasValidCaptureSettings = true;
+			}
+		});
 	
-	for (const FMediaFrameworkCaptureRenderTargetCameraOutputInfo& RenderTargetOutputInfo : CaptureSettings->RenderTargetCaptures)
-	{
-		if (RenderTargetOutputInfo.MediaOutput != InMediaOutput)
+	ForEachOutputInfoForMediaOutput<FMediaFrameworkCaptureRenderTargetCameraOutputInfo>(CaptureSettings, InMediaOutput, 
+		[this, &InMediaOutput, &bHasValidCaptureSettings](const FMediaFrameworkCaptureRenderTargetCameraOutputInfo& OutputInfo)
 		{
-			continue;
-		}
-
-		if (RenderTargetOutputInfo.RenderTarget)
+			if (OutputInfo.RenderTarget)
+			{
+				bHasValidCaptureSettings = true;
+			}
+		});
+	
+	ForEachOutputInfoForMediaOutput<FMediaFrameworkCaptureMediaTextureOutputInfo>(CaptureSettings, InMediaOutput, 
+		[this, &bHasValidCaptureSettings](const FMediaFrameworkCaptureMediaTextureOutputInfo& OutputInfo)
 		{
-			bHasValidCaptureSettings = true;
-			break;
-		}
-	}
+			if (OutputInfo.MediaTexture)
+			{
+				bHasValidCaptureSettings = true;
+			}
+		});
 
 	return bHasValidCaptureSettings;
 }
@@ -419,32 +422,7 @@ void FMediaProfileEditor::RefreshSelectedMediaItems(const TArray<int32>& InMedia
 		const int32 OutputIndex = InMediaOutputs[Index];
 		if (UMediaOutput* MediaOutput = MediaProfileBeingEdited->GetMediaOutput(OutputIndex))
 		{
-			if (PlaybackManager->IsOutputCapturing(MediaOutput))
-			{
-				if (UMediaProfileEditorCaptureSettings* CaptureSettings = GetMediaFrameworkCaptureSettings())
-				{
-					if (CaptureSettings->CurrentViewportMediaOutput.MediaOutput == MediaOutput)
-					{
-						PlaybackManager->RestartActiveViewportOutput(MediaOutput, CaptureSettings->CurrentViewportMediaOutput.CaptureOptions, CaptureSettings->bAutoRestartCaptureOnChange);
-					}
-
-					for (const FMediaFrameworkCaptureCameraViewportCameraOutputInfo& ViewportCapture : CaptureSettings->ViewportCaptures)
-					{
-						if (ViewportCapture.MediaOutput == MediaOutput)
-						{
-							PlaybackManager->RestartManagedViewportOutput(MediaOutput, ViewportCapture.CaptureOptions, CaptureSettings->bAutoRestartCaptureOnChange);
-						}
-					}
-
-					for (const FMediaFrameworkCaptureRenderTargetCameraOutputInfo& RenderTargetCapture : CaptureSettings->RenderTargetCaptures)
-					{
-						if (RenderTargetCapture.MediaOutput == MediaOutput)
-						{
-							PlaybackManager->RestartRenderTargetOutput(MediaOutput, RenderTargetCapture.RenderTarget, RenderTargetCapture.CaptureOptions, CaptureSettings->bAutoRestartCaptureOnChange);
-						}
-					}
-				}
-			}
+			RestartActiveMediaCaptures(MediaOutput);
 		}
 	}
 }
@@ -494,6 +472,52 @@ void FMediaProfileEditor::CleanUpAndSaveUserSettings()
 	UserSettings->SaveConfig();
 }
 
+void FMediaProfileEditor::RestartActiveMediaCaptures(UMediaOutput* InMediaOutput)
+{
+	UMediaProfilePlaybackManager* PlaybackManager = MediaProfileBeingEdited->GetPlaybackManager();
+	if (!PlaybackManager)
+	{
+		return;
+	}
+	
+	if (!PlaybackManager->IsOutputCapturing(InMediaOutput))
+	{
+		return;
+	}
+	
+	UMediaProfileEditorCaptureSettings* CaptureSettings = GetMediaFrameworkCaptureSettings();
+	if (!CaptureSettings)
+	{
+		return;
+	}
+	
+	using namespace UE::MediaFrameworkWorldSettings::Helpers;
+				
+	ForEachOutputInfoForMediaOutput<FMediaFrameworkCaptureCurrentViewportOutputInfo>(CaptureSettings, InMediaOutput, 
+		[&PlaybackManager, bAutoRestart = CaptureSettings->bAutoRestartCaptureOnChange](const FMediaFrameworkCaptureCurrentViewportOutputInfo& OutputInfo)
+		{
+			PlaybackManager->RestartActiveViewportOutput(OutputInfo.MediaOutput, OutputInfo.CaptureOptions, bAutoRestart);
+		});
+
+	ForEachOutputInfoForMediaOutput<FMediaFrameworkCaptureCameraViewportCameraOutputInfo>(CaptureSettings, InMediaOutput, 
+		[&PlaybackManager, bAutoRestart = CaptureSettings->bAutoRestartCaptureOnChange](const FMediaFrameworkCaptureCameraViewportCameraOutputInfo& OutputInfo)
+		{
+			PlaybackManager->RestartManagedViewportOutput(OutputInfo.MediaOutput, OutputInfo.CaptureOptions, bAutoRestart);
+		});
+
+	ForEachOutputInfoForMediaOutput<FMediaFrameworkCaptureRenderTargetCameraOutputInfo>(CaptureSettings, InMediaOutput, 
+		[&PlaybackManager, bAutoRestart = CaptureSettings->bAutoRestartCaptureOnChange](const FMediaFrameworkCaptureRenderTargetCameraOutputInfo& OutputInfo)
+		{
+			PlaybackManager->RestartRenderTargetOutput(OutputInfo.MediaOutput, OutputInfo.RenderTarget, OutputInfo.CaptureOptions, bAutoRestart);
+		});
+
+	ForEachOutputInfoForMediaOutput<FMediaFrameworkCaptureMediaTextureOutputInfo>(CaptureSettings, InMediaOutput, 
+		[&PlaybackManager, bAutoRestart = CaptureSettings->bAutoRestartCaptureOnChange](const FMediaFrameworkCaptureMediaTextureOutputInfo& OutputInfo)
+		{
+			PlaybackManager->RestartMediaTextureOutput(OutputInfo.MediaOutput, OutputInfo.MediaTexture, OutputInfo.CaptureOptions, bAutoRestart);
+		});
+}
+
 void FMediaProfileEditor::OnMediaItemDeleted(UClass* InMediaType, int32 InMediaItemIndex)
 {
 	if (InMediaType->IsChildOf<UMediaSource>())
@@ -652,32 +676,7 @@ void FMediaProfileEditor::OnObjectPropertyChanged(UObject* InObject, FPropertyCh
 			const int32 OutputIndex = MediaProfileBeingEdited->FindMediaOutputIndex(MediaOutput);
 			if (OutputIndex != INDEX_NONE)
 			{
-				if (PlaybackManager->IsOutputCapturing(MediaOutput))
-				{
-					if (UMediaProfileEditorCaptureSettings* CaptureSettings = GetMediaFrameworkCaptureSettings())
-					{
-						if (CaptureSettings->CurrentViewportMediaOutput.MediaOutput == MediaOutput)
-						{
-							PlaybackManager->RestartActiveViewportOutput(MediaOutput, CaptureSettings->CurrentViewportMediaOutput.CaptureOptions, CaptureSettings->bAutoRestartCaptureOnChange);
-						}
-
-						for (const FMediaFrameworkCaptureCameraViewportCameraOutputInfo& ViewportCapture : CaptureSettings->ViewportCaptures)
-						{
-							if (ViewportCapture.MediaOutput == MediaOutput)
-							{
-								PlaybackManager->RestartManagedViewportOutput(MediaOutput, ViewportCapture.CaptureOptions, CaptureSettings->bAutoRestartCaptureOnChange);
-							}
-						}
-
-						for (const FMediaFrameworkCaptureRenderTargetCameraOutputInfo& RenderTargetCapture : CaptureSettings->RenderTargetCaptures)
-						{
-							if (RenderTargetCapture.MediaOutput == MediaOutput)
-							{
-								PlaybackManager->RestartRenderTargetOutput(MediaOutput, RenderTargetCapture.RenderTarget, RenderTargetCapture.CaptureOptions, CaptureSettings->bAutoRestartCaptureOnChange);
-							}
-						}
-					}
-				}
+				RestartActiveMediaCaptures(MediaOutput);
 			}
 		}
 	}
diff --git a/Engine/Plugins/Media/MediaFrameworkUtilities/Source/MediaFrameworkUtilitiesEditor/Private/AssetEditor/MediaProfileEditor.h b/Engine/Plugins/Media/MediaFrameworkUtilities/Source/MediaFrameworkUtilitiesEditor/Private/AssetEditor/MediaProfileEditor.h
--- a/Engine/Plugins/Media/MediaFrameworkUtilities/Source/MediaFrameworkUtilitiesEditor/Private/AssetEditor/MediaProfileEditor.h
+++ b/Engine/Plugins/Media/MediaFrameworkUtilities/Source/MediaFrameworkUtilitiesEditor/Private/AssetEditor/MediaProfileEditor.h
@@ -88,6 +88,9 @@ private:
 	/** Cleans up and saves the media profile editor user settings config */
 	void CleanUpAndSaveUserSettings();
 	
+	/** Restart any active media captures for the specified media output */
+	void RestartActiveMediaCaptures(UMediaOutput* InMediaOutput);
+	
 	void OnMediaItemDeleted(UClass* InMediaType, int32 InMediaItemIndex);
 	void OnSelectedMediaItemsChanged(const TArray<int32>& SelectedMediaSources, const TArray<int32>& SelectedMediaOutputs);
 	
diff --git a/Engine/Plugins/Media/MediaFrameworkUtilities/Source/MediaFrameworkUtilitiesEditor/Private/AssetEditor/MediaProfileMediaItemDetailCustomization.cpp b/Engine/Plugins/Media/MediaFrameworkUtilities/Source/MediaFrameworkUtilitiesEditor/Private/AssetEditor/MediaProfileMediaItemDetailCustomization.cpp
--- a/Engine/Plugins/Media/MediaFrameworkUtilities/Source/MediaFrameworkUtilitiesEditor/Private/AssetEditor/MediaProfileMediaItemDetailCustomization.cpp
+++ b/Engine/Plugins/Media/MediaFrameworkUtilities/Source/MediaFrameworkUtilitiesEditor/Private/AssetEditor/MediaProfileMediaItemDetailCustomization.cpp
@@ -10,11 +10,11 @@
 #include "MediaPlayer.h"
 #include "MediaProfileEditorUserSettings.h"
 #include "MediaTexture.h"
+#include "SMediaProfileSourceTexturePicker.h"
 #include "StreamMediaSource.h"
-#include "Algo/Count.h"
 #include "Internationalization/Culture.h"
 #include "Profile/MediaProfilePlaybackManager.h"
-#include "UI/MediaFrameworkUtilitiesEditorStyle.h"
+#include "Widgets/SToolTip.h"
 #include "Widgets/Input/SComboBox.h"
 
 #define LOCTEXT_NAMESPACE "FMediaProfileMediaItemDetailCustomization"
@@ -610,6 +610,11 @@ public:
 					LOCTEXT("MediaCaptureRenderTargetMethodToolTip", "Capture from a render target"),
 					FSlateIcon(),
 					FUIAction(FExecuteAction::CreateSP(this, &SMediaCaptureMethodComboButton::OnRenderTargetMethodSelected)));
+
+				MenuBuilder.AddMenuEntry(LOCTEXT("MediaCaptureMediaTextureMethod", "Media Texture"),
+					LOCTEXT("MediaCaptureMediaTextureMethodToolTip", "Capture from a media texture (passthrough)"),
+					FSlateIcon(),
+					FUIAction(FExecuteAction::CreateSP(this, &SMediaCaptureMethodComboButton::OnMediaTextureMethodSelected)));
 				
 				return MenuBuilder.MakeWidget();
 			})
@@ -645,22 +650,12 @@ private:
 				continue;
 			}
 			
-			// Now, if the object has an existing viewport or render target capture config, we want to copy that config over to the current viewport config, so find it
-			if (FMediaFrameworkCaptureCameraViewportCameraOutputInfo* ExistingViewportOutputInfo =
-				MediaCaptureSettings->ViewportCaptures.FindByPredicate(GetFindOutputInfoPredicate<FMediaFrameworkCaptureCameraViewportCameraOutputInfo>(MediaOutput)))
-			{
-				MediaCaptureSettings->CurrentViewportMediaOutput.CaptureOptions = ExistingViewportOutputInfo->CaptureOptions;
-			}
-			else if (FMediaFrameworkCaptureRenderTargetCameraOutputInfo* ExistingRenderTargetOutputInfo =
-				MediaCaptureSettings->RenderTargetCaptures.FindByPredicate(GetFindOutputInfoPredicate<FMediaFrameworkCaptureRenderTargetCameraOutputInfo>(MediaOutput)))
-			{
-				MediaCaptureSettings->CurrentViewportMediaOutput.CaptureOptions = ExistingRenderTargetOutputInfo->CaptureOptions;
-			}
+			// Now, if the object has an existing capture config, we want to copy that config over to the current viewport config, so find it
+			MediaCaptureSettings->CurrentViewportMediaOutput.CaptureOptions = FindExistingMediaCaptureOptions(MediaOutput, MediaCaptureSettings);
 
 			// We expect there to be only one capture configuration per media output, so clear
-			// any existing viewport and render target capture configurations before adding the new one
-			MediaCaptureSettings->ViewportCaptures.RemoveAll(GetFindOutputInfoPredicate<FMediaFrameworkCaptureCameraViewportCameraOutputInfo>(MediaOutput));
-			MediaCaptureSettings->RenderTargetCaptures.RemoveAll(GetFindOutputInfoPredicate<FMediaFrameworkCaptureRenderTargetCameraOutputInfo>(MediaOutput));
+			// any existing capture configurations in the media capture settings before adding the new one
+			RemoveExistingMediaCaptures(MediaOutput, MediaCaptureSettings);
 			
 			MediaCaptureSettings->CurrentViewportMediaOutput.MediaOutput = MediaOutput;
 		}
@@ -693,27 +688,16 @@ private:
 			
 			// We expect there to be only one capture configuration per media output, so clear
 			// any existing viewport capture configurations before adding the new one
-			MediaCaptureSettings->ViewportCaptures.RemoveAll(GetFindOutputInfoPredicate<FMediaFrameworkCaptureCameraViewportCameraOutputInfo>(MediaOutput));
+			UE::MediaFrameworkWorldSettings::Helpers::RemoveAllOutputInfoForMediaOutput<FMediaFrameworkCaptureCameraViewportCameraOutputInfo>(MediaCaptureSettings, MediaOutput);
 
 			FMediaFrameworkCaptureCameraViewportCameraOutputInfo NewViewportOutputInfo;
 			NewViewportOutputInfo.MediaOutput = MediaOutput;
 			
 			// Now, if the object has an existing render target capture config, we want to copy that config over to the new viewport config, so find it
-			FMediaFrameworkCaptureRenderTargetCameraOutputInfo* ExistingRenderTargetOutputInfo =
-				MediaCaptureSettings->RenderTargetCaptures.FindByPredicate(GetFindOutputInfoPredicate<FMediaFrameworkCaptureRenderTargetCameraOutputInfo>(MediaOutput));
-			
-			if (ExistingRenderTargetOutputInfo)
-			{
-				NewViewportOutputInfo.CaptureOptions = ExistingRenderTargetOutputInfo->CaptureOptions;
-			}
-
-			if (MediaCaptureSettings->CurrentViewportMediaOutput.MediaOutput == MediaOutput)
-			{
-				MediaCaptureSettings->CurrentViewportMediaOutput.MediaOutput = nullptr;
-			}
+			NewViewportOutputInfo.CaptureOptions = FindExistingMediaCaptureOptions(MediaOutput, MediaCaptureSettings);
 
+			RemoveExistingMediaCaptures(MediaOutput, MediaCaptureSettings);
 			MediaCaptureSettings->ViewportCaptures.Add(NewViewportOutputInfo);
-			MediaCaptureSettings->RenderTargetCaptures.RemoveAll(GetFindOutputInfoPredicate<FMediaFrameworkCaptureRenderTargetCameraOutputInfo>(MediaOutput));
 		}
 
 		MediaCaptureSettings->SaveConfig();
@@ -744,27 +728,62 @@ private:
 
 			// We expect there to be only one capture configuration per media output, so clear
 			// any existing render target capture configurations before adding the new one
-			MediaCaptureSettings->RenderTargetCaptures.RemoveAll(GetFindOutputInfoPredicate<FMediaFrameworkCaptureRenderTargetCameraOutputInfo>(MediaOutput));
+			UE::MediaFrameworkWorldSettings::Helpers::RemoveAllOutputInfoForMediaOutput<FMediaFrameworkCaptureRenderTargetCameraOutputInfo>(MediaCaptureSettings, MediaOutput);
 
 			FMediaFrameworkCaptureRenderTargetCameraOutputInfo NewRenderTargetOutputInfo;
 			NewRenderTargetOutputInfo.MediaOutput = MediaOutput;
 			
 			// Now, if the object has an existing render target capture config, we want to copy that config over to the new viewport config, so find it
-			FMediaFrameworkCaptureCameraViewportCameraOutputInfo* ExistingViewportOutputInfo =
-				MediaCaptureSettings->ViewportCaptures.FindByPredicate(GetFindOutputInfoPredicate<FMediaFrameworkCaptureCameraViewportCameraOutputInfo>(MediaOutput));
-			
-			if (ExistingViewportOutputInfo)
+			NewRenderTargetOutputInfo.CaptureOptions = FindExistingMediaCaptureOptions(MediaOutput, MediaCaptureSettings);
+
+			RemoveExistingMediaCaptures(MediaOutput, MediaCaptureSettings);
+			MediaCaptureSettings->RenderTargetCaptures.Add(NewRenderTargetOutputInfo);
+		}
+
+		MediaCaptureSettings->SaveConfig();
+		OnMethodChanged.ExecuteIfBound();
+	}
+
+	void OnMediaTextureMethodSelected()
+	{
+		UMediaProfileEditorCaptureSettings* MediaCaptureSettings = FMediaProfileEditor::GetMediaFrameworkCaptureSettings();
+		if (!MediaCaptureSettings)
+		{
+			return;
+		}
+		
+		for (const TWeakObjectPtr<UObject>& WeakObject : Objects)
+		{
+			TStrongObjectPtr<UObject> Object = WeakObject.Pin();
+			UMediaOutput* MediaOutput = Object.IsValid() ? Cast<UMediaOutput>(Object.Get()) : nullptr;
+			if (!MediaOutput)
 			{
-				NewRenderTargetOutputInfo.CaptureOptions = ExistingViewportOutputInfo->CaptureOptions;
+				continue;
 			}
 
-			if (MediaCaptureSettings->CurrentViewportMediaOutput.MediaOutput == MediaOutput)
+			if (MediaOutput->IsA<UDummyMediaOutput>())
 			{
-				MediaCaptureSettings->CurrentViewportMediaOutput.MediaOutput = nullptr;
+				continue;
 			}
 
-			MediaCaptureSettings->RenderTargetCaptures.Add(NewRenderTargetOutputInfo);
-			MediaCaptureSettings->ViewportCaptures.RemoveAll(GetFindOutputInfoPredicate<FMediaFrameworkCaptureCameraViewportCameraOutputInfo>(MediaOutput));
+			// We expect there to be only one capture configuration per media output, so clear
+			// any existing render target capture configurations before adding the new one
+			UE::MediaFrameworkWorldSettings::Helpers::RemoveAllOutputInfoForMediaOutput<FMediaFrameworkCaptureMediaTextureOutputInfo>(MediaCaptureSettings, MediaOutput);
+
+			FMediaFrameworkCaptureMediaTextureOutputInfo NewMediaTextureOutputInfo;
+			NewMediaTextureOutputInfo.MediaOutput = MediaOutput;
+			
+			// Now, if the object has an existing render target capture config, we want to copy that config over to the new viewport config, so find it
+			NewMediaTextureOutputInfo.CaptureOptions = FindExistingMediaCaptureOptions(MediaOutput, MediaCaptureSettings);
+
+			if (NewMediaTextureOutputInfo.CaptureOptions.ResizeMethod == EMediaCaptureResizeMethod::ResizeSource)
+			{
+				// ResizeSource is not supported for media textures, so revert to None
+				NewMediaTextureOutputInfo.CaptureOptions.ResizeMethod = EMediaCaptureResizeMethod::None;
+			}
+			
+			RemoveExistingMediaCaptures(MediaOutput, MediaCaptureSettings);
+			MediaCaptureSettings->MediaTextureCaptures.Add(NewMediaTextureOutputInfo);
 		}
 
 		MediaCaptureSettings->SaveConfig();
@@ -779,9 +798,10 @@ private:
 			return LOCTEXT("MediaCaptureMethodNotConfigured", "Not Configured");
 		}
 		
-		bool bCurrentViewport = false;
+		int32 NumCurrentViewports = 0;
 		int32 NumViewports = 0;
 		int32 NumRenderTargets = 0;
+		int32 NumMediaTextures = 0;
 
 		for (const TWeakObjectPtr<UObject>& WeakObject : Objects)
 		{
@@ -792,17 +812,19 @@ private:
 				continue;
 			}
 
-			bCurrentViewport = MediaCaptureSettings->CurrentViewportMediaOutput.MediaOutput == MediaOutput;
-			NumViewports += Algo::CountIf(MediaCaptureSettings->ViewportCaptures, GetFindOutputInfoPredicate<FMediaFrameworkCaptureCameraViewportCameraOutputInfo>(MediaOutput));
-			NumRenderTargets += Algo::CountIf(MediaCaptureSettings->RenderTargetCaptures, GetFindOutputInfoPredicate<FMediaFrameworkCaptureRenderTargetCameraOutputInfo>(MediaOutput));
+			using namespace UE::MediaFrameworkWorldSettings::Helpers;
+			NumCurrentViewports += FindAllOutputInfoForMediaOutput<FMediaFrameworkCaptureCurrentViewportOutputInfo>(MediaCaptureSettings, MediaOutput).Num();
+			NumViewports += FindAllOutputInfoForMediaOutput<FMediaFrameworkCaptureCameraViewportCameraOutputInfo>(MediaCaptureSettings, MediaOutput).Num();
+			NumRenderTargets += FindAllOutputInfoForMediaOutput<FMediaFrameworkCaptureRenderTargetCameraOutputInfo>(MediaCaptureSettings, MediaOutput).Num();
+			NumMediaTextures += FindAllOutputInfoForMediaOutput<FMediaFrameworkCaptureMediaTextureOutputInfo>(MediaCaptureSettings, MediaOutput).Num();
 		}
 
-		const bool bMultipleTypes = (bCurrentViewport && (NumViewports > 0 || NumRenderTargets > 0)) || (NumViewports > 0 && NumRenderTargets > 0);
-		if (bMultipleTypes)
+		const int32 NumTypesWithOutputInfo = (int32)(NumCurrentViewports > 0) + (int32)(NumViewports > 0) + (int32)(NumRenderTargets > 0) + (int32)(NumMediaTextures > 0);
+		if (NumTypesWithOutputInfo >= 2)
 		{
 			return LOCTEXT("MediaCaptureMethodMultipleValues", "Multiple Values");
 		}
-		if (bCurrentViewport)
+		if (NumCurrentViewports > 0)
 		{
 			return LOCTEXT("MediaCaptureCurrentViewportMethod", "Current Viewport");
 		}
@@ -814,17 +836,53 @@ private:
 		{
 			return LOCTEXT("MediaCaptureRenderTargetMethod", "Render Target");
 		}
+		if (NumMediaTextures > 0)
+		{
+			return LOCTEXT("MediaCaptureMediaTextureMethod", "Media Texture");
+		}
 
 		return LOCTEXT("MediaCaptureMethodNotConfigured", "Not Configured");
 	}
 	
-	template<typename TOutputInfo>
-	TFunction<bool(const TOutputInfo&)> GetFindOutputInfoPredicate(UMediaOutput* InMediaOutput) const
+	FMediaCaptureOptions FindExistingMediaCaptureOptions(UMediaOutput* InMediaOutput, UMediaProfileEditorCaptureSettings* MediaCaptureSettings) const
 	{
-		return [InMediaOutput](const TOutputInfo& OutputInfo)
+		using namespace UE::MediaFrameworkWorldSettings::Helpers;
+		
+		if (FMediaFrameworkCaptureCurrentViewportOutputInfo* CurrentViewportCapture =
+			FindFirstOutputInfoForMediaOutput<FMediaFrameworkCaptureCurrentViewportOutputInfo>(MediaCaptureSettings, InMediaOutput))
 		{
-			return OutputInfo.MediaOutput == InMediaOutput;
-		};
+			return CurrentViewportCapture->CaptureOptions;
+		}
+		
+		if (FMediaFrameworkCaptureCameraViewportCameraOutputInfo* ViewportCapture =
+			FindFirstOutputInfoForMediaOutput<FMediaFrameworkCaptureCameraViewportCameraOutputInfo>(MediaCaptureSettings, InMediaOutput))
+		{
+			return ViewportCapture->CaptureOptions;
+		}
+		
+		if (FMediaFrameworkCaptureRenderTargetCameraOutputInfo* RenderTargetCapture =
+			FindFirstOutputInfoForMediaOutput<FMediaFrameworkCaptureRenderTargetCameraOutputInfo>(MediaCaptureSettings, InMediaOutput))
+		{
+			return RenderTargetCapture->CaptureOptions;
+		}
+		
+		if (FMediaFrameworkCaptureMediaTextureOutputInfo* MediaTextureCapture =
+			FindFirstOutputInfoForMediaOutput<FMediaFrameworkCaptureMediaTextureOutputInfo>(MediaCaptureSettings, InMediaOutput))
+		{
+			return MediaTextureCapture->CaptureOptions;
+		}
+		
+		return FMediaCaptureOptions();
+	}
+
+	void RemoveExistingMediaCaptures(UMediaOutput* InMediaOutput, UMediaProfileEditorCaptureSettings* MediaCaptureSettings)
+	{
+		using namespace UE::MediaFrameworkWorldSettings::Helpers;
+		
+		RemoveAllOutputInfoForMediaOutput<FMediaFrameworkCaptureCurrentViewportOutputInfo>(MediaCaptureSettings, InMediaOutput);
+		RemoveAllOutputInfoForMediaOutput<FMediaFrameworkCaptureCameraViewportCameraOutputInfo>(MediaCaptureSettings, InMediaOutput);
+		RemoveAllOutputInfoForMediaOutput<FMediaFrameworkCaptureRenderTargetCameraOutputInfo>(MediaCaptureSettings, InMediaOutput);
+		RemoveAllOutputInfoForMediaOutput<FMediaFrameworkCaptureMediaTextureOutputInfo>(MediaCaptureSettings, InMediaOutput);
 	}
 	
 private:
@@ -832,13 +890,44 @@ private:
 	FSimpleDelegate OnMethodChanged;
 };
 
-class FMediaFrameworkCurrentViewportOutputInfoCustomization : public IPropertyTypeCustomization
+class FMediaFrameworkBaseOutputInfoCustomization : public IPropertyTypeCustomization
 {
+public:
+	FMediaFrameworkBaseOutputInfoCustomization(const TWeakPtr<FMediaProfileEditor>& InMediaProfileEditor)
+		: MediaProfileEditor(InMediaProfileEditor)
+	{ }
+	
 protected:
 	virtual void CustomizeHeader(TSharedRef<IPropertyHandle> PropertyHandle, FDetailWidgetRow& HeaderRow, IPropertyTypeCustomizationUtils& CustomizationUtils) override { }
 
 	virtual void CustomizeChildren(TSharedRef<IPropertyHandle> PropertyHandle, IDetailChildrenBuilder& ChildBuilder, IPropertyTypeCustomizationUtils& CustomizationUtils) override
 	{
+		TMap<FName, TFunction<void(TSharedPtr<IPropertyHandle>)>> PropertyCustomizers =
+		{
+			{
+				GET_MEMBER_NAME_CHECKED(FMediaCaptureOptions, ResizeMethod),
+				[this, &ChildBuilder](TSharedPtr<IPropertyHandle> PropertyHandle)
+				{
+					FPropertyComboBoxArgs ComboBoxArgs;
+					ComboBoxArgs.PropertyHandle = PropertyHandle;
+					ComboBoxArgs.OnGetStrings = FOnGetPropertyComboBoxStrings::CreateSP(this, &FMediaFrameworkBaseOutputInfoCustomization::GetResizeMethodComboBoxStrings);
+					ComboBoxArgs.OnGetValue = FOnGetPropertyComboBoxValue::CreateSP(this, &FMediaFrameworkBaseOutputInfoCustomization::GetResizeMethodComboBoxValue, PropertyHandle);
+					ComboBoxArgs.OnValueSelected = FOnPropertyComboBoxValueSelected::CreateSP(this, &FMediaFrameworkBaseOutputInfoCustomization::SetResizeMethodValue, PropertyHandle);
+					
+					ChildBuilder.AddCustomRow(PropertyHandle->GetPropertyDisplayName())
+					.PropertyHandleList({ PropertyHandle })
+					.NameContent()
+					[
+						PropertyHandle->CreatePropertyNameWidget()
+					]
+					.ValueContent()
+					[
+						PropertyCustomizationHelpers::MakePropertyComboBox(ComboBoxArgs)
+					];
+				}
+			}
+		};
+		
 		// Manually add all capture options properties to avoid an additional property group getting added
 		TSharedPtr<IPropertyHandle> CaptureOptionsHandle = PropertyHandle->GetChildHandle(GET_MEMBER_NAME_CHECKED(FMediaFrameworkCaptureCurrentViewportOutputInfo, CaptureOptions));
 		uint32 NumChildren;
@@ -846,12 +935,22 @@ protected:
 		for (uint32 Index = 0; Index < NumChildren; ++Index)
 		{
 			TSharedPtr<IPropertyHandle> ChildHandle = CaptureOptionsHandle->GetChildHandle(Index);
-			ChildHandle->SetOnPropertyValueChangedWithData(TDelegate<void(const FPropertyChangedEvent&)>::CreateSP(this, &FMediaFrameworkCurrentViewportOutputInfoCustomization::PropertyValueChange));			
-			ChildBuilder.AddProperty(ChildHandle.ToSharedRef());
+			ChildHandle->SetOnPropertyValueChangedWithData(TDelegate<void(const FPropertyChangedEvent&)>::CreateSP(this, &FMediaFrameworkBaseOutputInfoCustomization::PropertyValueChange));
+			
+			const FName ChildPropertyName = ChildHandle->GetProperty()->GetFName();
+			if (PropertyCustomizers.Contains(ChildPropertyName))
+			{
+				PropertyCustomizers[ChildPropertyName](ChildHandle);
+			}
+			else
+			{
+				ChildBuilder.AddProperty(ChildHandle.ToSharedRef());
+			}
 		}
 	}
 	
-	void PropertyValueChange(const FPropertyChangedEvent& PropertyChangedEvent)
+	/** Raised when any of the capture info properties have changed */
+	virtual void PropertyValueChange(const FPropertyChangedEvent& PropertyChangedEvent)
 	{
 		if (UMediaProfileEditorCaptureSettings* MediaCaptureSettings = FMediaProfileEditor::GetMediaFrameworkCaptureSettings())
 		{
@@ -859,71 +958,162 @@ protected:
 			FCoreUObjectDelegates::OnObjectPropertyChanged.Broadcast(MediaCaptureSettings, const_cast<FPropertyChangedEvent&>(PropertyChangedEvent));
 		}
 	}
-};
 
-class FMediaFrameworkViewportOutputInfoCustomization : public IPropertyTypeCustomization
-{
-protected:
-	virtual void CustomizeHeader(TSharedRef<IPropertyHandle> PropertyHandle, FDetailWidgetRow& HeaderRow, IPropertyTypeCustomizationUtils& CustomizationUtils) override { }
-
-	virtual void CustomizeChildren(TSharedRef<IPropertyHandle> PropertyHandle, IDetailChildrenBuilder& ChildBuilder, IPropertyTypeCustomizationUtils& CustomizationUtils) override
+	/** Gets the EMediaCaptureResizeMethod display strings for the ResizeMethod combo box */
+	void GetResizeMethodComboBoxStrings(TArray<TSharedPtr<FString>>& OutStrings, TArray<TSharedPtr<SToolTip>>& OutToolTips, TArray<bool>& OutRestrictedItems)
 	{
-		TSharedPtr<IPropertyHandle> LockedActorsHandle = PropertyHandle->GetChildHandle(GET_MEMBER_NAME_CHECKED(FMediaFrameworkCaptureCameraViewportCameraOutputInfo, Cameras));
-		LockedActorsHandle->SetOnPropertyValueChangedWithData(TDelegate<void(const FPropertyChangedEvent&)>::CreateSP(this, &FMediaFrameworkViewportOutputInfoCustomization::PropertyValueChange));
-		LockedActorsHandle->SetOnChildPropertyValueChangedWithData(TDelegate<void(const FPropertyChangedEvent&)>::CreateSP(this, &FMediaFrameworkViewportOutputInfoCustomization::PropertyValueChange));
-		ChildBuilder.AddProperty(LockedActorsHandle.ToSharedRef());
+		const UEnum* ResizeMethodEnum = StaticEnum<EMediaCaptureResizeMethod>();
 
-		// Manually add all capture options properties to avoid an additional property group getting added
-		TSharedPtr<IPropertyHandle> CaptureOptionsHandle = PropertyHandle->GetChildHandle(GET_MEMBER_NAME_CHECKED(FMediaFrameworkCaptureCameraViewportCameraOutputInfo, CaptureOptions));
-		uint32 NumChildren;
-		CaptureOptionsHandle->GetNumChildren(NumChildren);
-		for (uint32 Index = 0; Index < NumChildren; ++Index)
+		const TArray<EMediaCaptureResizeMethod> RestrictedResizeMethods = GetRestrictedResizeMethods();
+		
+		const int32 NumEnums = ResizeMethodEnum->NumEnums() - 1;
+		for (int32 Entry = 0; Entry < NumEnums; ++Entry)
 		{
-			TSharedPtr<IPropertyHandle> ChildHandle = CaptureOptionsHandle->GetChildHandle(Index);
-			ChildHandle->SetOnPropertyValueChangedWithData(TDelegate<void(const FPropertyChangedEvent&)>::CreateSP(this, &FMediaFrameworkViewportOutputInfoCustomization::PropertyValueChange));			
-			ChildBuilder.AddProperty(ChildHandle.ToSharedRef());
+			FText EnumName = ResizeMethodEnum->GetDisplayNameTextByIndex(Entry);
+
+			OutStrings.Add(MakeShared<FString>(EnumName.ToString()));
+			OutToolTips.Add(SNew(SToolTip).Text(ResizeMethodEnum->GetToolTipTextByIndex(Entry)));
+			OutRestrictedItems.Add(RestrictedResizeMethods.Contains((EMediaCaptureResizeMethod)Entry));
 		}
 	}
+
+	/** Gets the current value EMediaCaptureResizeMethod display string for the ResizeMethod combo box */
+	FString GetResizeMethodComboBoxValue(TSharedPtr<IPropertyHandle> PropertyHandle)
+	{
+		const UEnum* ResizeMethodEnum = StaticEnum<EMediaCaptureResizeMethod>();
+		
+		uint8 Value;
+		if (PropertyHandle->GetValue(Value) == FPropertyAccess::Success)
+		{
+			return ResizeMethodEnum->GetDisplayNameTextByIndex(Value).ToString();
+		}
+		
+		return TEXT("Failed to get value");
+	}
 	
-	void PropertyValueChange(const FPropertyChangedEvent& PropertyChangedEvent)
+	/** Sets the ResizeMethod value from the selected option from the combo box*/
+	void SetResizeMethodValue(const FString& InValue, TSharedPtr<IPropertyHandle> PropertyHandle)
 	{
-		if (UMediaProfileEditorCaptureSettings* MediaCaptureSettings = FMediaProfileEditor::GetMediaFrameworkCaptureSettings())
+		uint8 Value;
+		PropertyHandle->GetValue(Value);
+		
+		const UEnum* ResizeMethodEnum = StaticEnum<EMediaCaptureResizeMethod>();
+		int32 EnumIndex = INDEX_NONE;
+		const int32 NumEnums = ResizeMethodEnum->NumEnums() - 1;
+		for (int32 Entry = 0; Entry < NumEnums; ++Entry)
 		{
-			MediaCaptureSettings->SaveConfig();
-			FCoreUObjectDelegates::OnObjectPropertyChanged.Broadcast(MediaCaptureSettings, const_cast<FPropertyChangedEvent&>(PropertyChangedEvent));
+			FText EnumName = ResizeMethodEnum->GetDisplayNameTextByIndex(Entry);
+
+			if (EnumName.ToString() == InValue)
+			{
+				EnumIndex = Entry;
+				break;
+			}
+		}
+		
+		if (Value != (uint8)EnumIndex)
+		{
+			PropertyHandle->SetValue((uint8)EnumIndex);
 		}
 	}
+	
+	/** Gets the list of EMediaCaptureResizeMethod options that should be restricted for the current customization */
+	virtual TArray<EMediaCaptureResizeMethod> GetRestrictedResizeMethods() const
+	{
+		return TArray<EMediaCaptureResizeMethod>();
+	}
+	
+protected:
+	/** Weak pointer to the  media profile editor */
+	TWeakPtr<FMediaProfileEditor> MediaProfileEditor;
 };
 
-class FMediaFrameworkRenderTargetOutputInfoCustomization : public IPropertyTypeCustomization
+class FMediaFrameworkViewportOutputInfoCustomization : public FMediaFrameworkBaseOutputInfoCustomization
 {
+public:
+	FMediaFrameworkViewportOutputInfoCustomization(const TWeakPtr<FMediaProfileEditor>& InMediaProfileEditor)
+		: FMediaFrameworkBaseOutputInfoCustomization(InMediaProfileEditor)
+	{ }
+
 protected:
-	virtual void CustomizeHeader(TSharedRef<IPropertyHandle> PropertyHandle, FDetailWidgetRow& HeaderRow, IPropertyTypeCustomizationUtils& CustomizationUtils) override { }
+	virtual void CustomizeChildren(TSharedRef<IPropertyHandle> PropertyHandle, IDetailChildrenBuilder& ChildBuilder, IPropertyTypeCustomizationUtils& CustomizationUtils) override
+	{
+		TSharedPtr<IPropertyHandle> LockedActorsHandle = PropertyHandle->GetChildHandle(GET_MEMBER_NAME_CHECKED(FMediaFrameworkCaptureCameraViewportCameraOutputInfo, Cameras));
+		LockedActorsHandle->SetOnPropertyValueChangedWithData(TDelegate<void(const FPropertyChangedEvent&)>::CreateSP(this, &FMediaFrameworkViewportOutputInfoCustomization::PropertyValueChange));
+		LockedActorsHandle->SetOnChildPropertyValueChangedWithData(TDelegate<void(const FPropertyChangedEvent&)>::CreateSP(this, &FMediaFrameworkViewportOutputInfoCustomization::PropertyValueChange));
+		ChildBuilder.AddProperty(LockedActorsHandle.ToSharedRef());
+
+		FMediaFrameworkBaseOutputInfoCustomization::CustomizeChildren(PropertyHandle, ChildBuilder, CustomizationUtils);
+	}
+};
+
+class FMediaFrameworkRenderTargetOutputInfoCustomization : public FMediaFrameworkBaseOutputInfoCustomization
+{
+public:
+	FMediaFrameworkRenderTargetOutputInfoCustomization(const TWeakPtr<FMediaProfileEditor>& InMediaProfileEditor)
+		: FMediaFrameworkBaseOutputInfoCustomization(InMediaProfileEditor)
+	{ }
 
+protected:
 	virtual void CustomizeChildren(TSharedRef<IPropertyHandle> PropertyHandle, IDetailChildrenBuilder& ChildBuilder, IPropertyTypeCustomizationUtils& CustomizationUtils) override
 	{
 		TSharedPtr<IPropertyHandle> RenderTargetHandle = PropertyHandle->GetChildHandle(GET_MEMBER_NAME_CHECKED(FMediaFrameworkCaptureRenderTargetCameraOutputInfo, RenderTarget));
 		RenderTargetHandle->SetOnPropertyValueChangedWithData(TDelegate<void(const FPropertyChangedEvent&)>::CreateSP(this, &FMediaFrameworkRenderTargetOutputInfoCustomization::PropertyValueChange));
 		ChildBuilder.AddProperty(RenderTargetHandle.ToSharedRef());
 
-		// Manually add all capture options properties to avoid an additional property group getting added
-		TSharedPtr<IPropertyHandle> CaptureOptionsHandle = PropertyHandle->GetChildHandle(GET_MEMBER_NAME_CHECKED(FMediaFrameworkCaptureRenderTargetCameraOutputInfo, CaptureOptions));
-		uint32 NumChildren;
-		CaptureOptionsHandle->GetNumChildren(NumChildren);
-		for (uint32 Index = 0; Index < NumChildren; ++Index)
-		{
-			TSharedPtr<IPropertyHandle> ChildHandle = CaptureOptionsHandle->GetChildHandle(Index);
-			ChildHandle->SetOnPropertyValueChangedWithData(TDelegate<void(const FPropertyChangedEvent&)>::CreateSP(this, &FMediaFrameworkRenderTargetOutputInfoCustomization::PropertyValueChange));
-			ChildBuilder.AddProperty(ChildHandle.ToSharedRef());
-		}
+		FMediaFrameworkBaseOutputInfoCustomization::CustomizeChildren(PropertyHandle, ChildBuilder, CustomizationUtils);
+	}
+};
+
+class FMediaFrameworkMediaTextureOutputInfoCustomization : public FMediaFrameworkBaseOutputInfoCustomization
+{
+public:
+	FMediaFrameworkMediaTextureOutputInfoCustomization(const TWeakPtr<FMediaProfileEditor>& InMediaProfileEditor)
+		: FMediaFrameworkBaseOutputInfoCustomization(InMediaProfileEditor)
+	{ }
+
+protected:
+	virtual void CustomizeChildren(TSharedRef<IPropertyHandle> PropertyHandle, IDetailChildrenBuilder& ChildBuilder, IPropertyTypeCustomizationUtils& CustomizationUtils) override
+	{
+		TSharedPtr<IPropertyHandle> MediaTextureHandle = PropertyHandle->GetChildHandle(GET_MEMBER_NAME_CHECKED(FMediaFrameworkCaptureMediaTextureOutputInfo, MediaTexture));
+		MediaTextureHandle->SetOnPropertyValueChangedWithData(TDelegate<void(const FPropertyChangedEvent&)>::CreateSP(this, &FMediaFrameworkMediaTextureOutputInfoCustomization::PropertyValueChange));
+
+		ChildBuilder.AddCustomRow(MediaTextureHandle->GetPropertyDisplayName())
+		.NameContent()
+		[
+			MediaTextureHandle->CreatePropertyNameWidget()
+		]
+		.ValueContent()
+		[
+			SNew(SMediaProfileSourceTexturePicker)
+			.TexturePropertyHandle(MediaTextureHandle)
+			.ThumbnailPool(CustomizationUtils.GetThumbnailPool())
+			.OnMediaSourceSelected(this, &FMediaFrameworkMediaTextureOutputInfoCustomization::OnMediaSourceSelected)
+		];
+
+		FMediaFrameworkBaseOutputInfoCustomization::CustomizeChildren(PropertyHandle, ChildBuilder, CustomizationUtils);
 	}
 
-	void PropertyValueChange(const FPropertyChangedEvent& InPropertyChangedEvent)
+	virtual TArray<EMediaCaptureResizeMethod> GetRestrictedResizeMethods() const override
 	{
-		if (UMediaProfileEditorCaptureSettings* MediaCaptureSettings = FMediaProfileEditor::GetMediaFrameworkCaptureSettings())
+		return 
 		{
-			MediaCaptureSettings->SaveConfig();
-			FCoreUObjectDelegates::OnObjectPropertyChanged.Broadcast(MediaCaptureSettings, const_cast<FPropertyChangedEvent&>(InPropertyChangedEvent));
+			EMediaCaptureResizeMethod::ResizeSource
+		};
+	}
+	
+private:
+	/** Raised when a media source is selected from the media profile media source options dropdown */
+	void OnMediaSourceSelected(UMediaProfile* InMediaProfile, int32 InMediaSourceIndex)
+	{
+		if (InMediaProfile == nullptr)
+		{
+			return;
+		}
+		
+		if (TSharedPtr<FMediaProfileEditor> PinnedMediaProfileEditor = MediaProfileEditor.Pin())
+		{
+			InMediaProfile->GetPlaybackManager()->OpenSourceFromIndex(InMediaSourceIndex, PinnedMediaProfileEditor.Get());
 		}
 	}
 };
@@ -949,23 +1139,30 @@ void FMediaProfileMediaOutputDetailCustomization::CustomizeDetails(IDetailLayout
 
 	DetailBuilder.RegisterInstancedCustomPropertyTypeLayout(
 		FMediaFrameworkCaptureCurrentViewportOutputInfo::StaticStruct()->GetFName(),
-		FOnGetPropertyTypeCustomizationInstance::CreateLambda([]
+		FOnGetPropertyTypeCustomizationInstance::CreateLambda([WeakMediaProfile = MediaProfileEditor]
 		{
-			return MakeShared<FMediaFrameworkCurrentViewportOutputInfoCustomization>();
+			return MakeShared<FMediaFrameworkBaseOutputInfoCustomization>(WeakMediaProfile);
 		}));
 	
 	DetailBuilder.RegisterInstancedCustomPropertyTypeLayout(
 		FMediaFrameworkCaptureCameraViewportCameraOutputInfo::StaticStruct()->GetFName(),
-		FOnGetPropertyTypeCustomizationInstance::CreateLambda([]
+		FOnGetPropertyTypeCustomizationInstance::CreateLambda([WeakMediaProfile = MediaProfileEditor]
 		{
-			return MakeShared<FMediaFrameworkViewportOutputInfoCustomization>();
+			return MakeShared<FMediaFrameworkViewportOutputInfoCustomization>(WeakMediaProfile);
 		}));
 
 	DetailBuilder.RegisterInstancedCustomPropertyTypeLayout(
 		FMediaFrameworkCaptureRenderTargetCameraOutputInfo::StaticStruct()->GetFName(),
-		FOnGetPropertyTypeCustomizationInstance::CreateLambda([]
+		FOnGetPropertyTypeCustomizationInstance::CreateLambda([WeakMediaProfile = MediaProfileEditor]
+		{
+			return MakeShared<FMediaFrameworkRenderTargetOutputInfoCustomization>(WeakMediaProfile);
+		}));
+
+	DetailBuilder.RegisterInstancedCustomPropertyTypeLayout(
+		FMediaFrameworkCaptureMediaTextureOutputInfo::StaticStruct()->GetFName(),
+		FOnGetPropertyTypeCustomizationInstance::CreateLambda([WeakMediaProfile = MediaProfileEditor]
 		{
-			return MakeShared<FMediaFrameworkRenderTargetOutputInfoCustomization>();
+			return MakeShared<FMediaFrameworkMediaTextureOutputInfoCustomization>(WeakMediaProfile);
 		}));
 	
 	// For media outputs, add the corresponding capture info to the details panel
@@ -1003,6 +1200,7 @@ void FMediaProfileMediaOutputDetailCustomization::CustomizeDetails(IDetailLayout
 	bool bHasCurrentViewportCapture = false;
 	bool bHasViewportCaptures = false;
 	bool bHasRenderTargetCaptures = false;
+	bool bHasMediaTextureCaptures = false;
 	if (UMediaProfileEditorCaptureSettings* MediaCaptureSettings = FMediaProfileEditor::GetMediaFrameworkCaptureSettings())
 	{
 		for (const TWeakObjectPtr<UObject>& WeakObject : Objects)
@@ -1019,46 +1217,55 @@ void FMediaProfileMediaOutputDetailCustomization::CustomizeDetails(IDetailLayout
 				continue;
 			}
 
-			if (MediaCaptureSettings->CurrentViewportMediaOutput.MediaOutput == MediaOutput)
-			{
-				TSharedPtr<FStructOnScope> Struct = MakeShared<FStructOnScope>(
-					FMediaFrameworkCaptureCurrentViewportOutputInfo::StaticStruct(),
-					reinterpret_cast<uint8*>(&MediaCaptureSettings->CurrentViewportMediaOutput));
-				Struct->SetPackage(MediaCaptureSettings->GetPackage());
-					
-				CaptureSettingStructs.Add(Struct);
-				bHasCurrentViewportCapture = true;
-			}
+			using namespace UE::MediaFrameworkWorldSettings::Helpers;
 			
-			for (FMediaFrameworkCaptureCameraViewportCameraOutputInfo& ViewportOutputInfo : MediaCaptureSettings->ViewportCaptures)
-			{
-				if (ViewportOutputInfo.MediaOutput == MediaOutput)
+			ForEachOutputInfoForMediaOutput<FMediaFrameworkCaptureCurrentViewportOutputInfo>(MediaCaptureSettings, MediaOutput,
+				[&MediaCaptureSettings, &CaptureSettingStructs, &bHasCurrentViewportCapture](FMediaFrameworkCaptureCurrentViewportOutputInfo& OutputInfo)
 				{
-					TSharedPtr<FStructOnScope> Struct = MakeShared<FStructOnScope>(FMediaFrameworkCaptureCameraViewportCameraOutputInfo::StaticStruct(), reinterpret_cast<uint8*>(&ViewportOutputInfo));
+					TSharedPtr<FStructOnScope> Struct = MakeShared<FStructOnScope>(
+						FMediaFrameworkCaptureCurrentViewportOutputInfo::StaticStruct(),
+						reinterpret_cast<uint8*>(&OutputInfo));
+					Struct->SetPackage(MediaCaptureSettings->GetPackage());
+						
+					CaptureSettingStructs.Add(Struct);
+					bHasCurrentViewportCapture = true;
+				});
+			
+			ForEachOutputInfoForMediaOutput<FMediaFrameworkCaptureCameraViewportCameraOutputInfo>(MediaCaptureSettings, MediaOutput,
+				[&MediaCaptureSettings, &CaptureSettingStructs, &bHasViewportCaptures](FMediaFrameworkCaptureCameraViewportCameraOutputInfo& OutputInfo)
+				{
+					TSharedPtr<FStructOnScope> Struct = MakeShared<FStructOnScope>(FMediaFrameworkCaptureCameraViewportCameraOutputInfo::StaticStruct(), reinterpret_cast<uint8*>(&OutputInfo));
 					Struct->SetPackage(MediaCaptureSettings->GetPackage());
 					
 					CaptureSettingStructs.Add(Struct);
 					bHasViewportCaptures = true;
-				}
-			}
+				});
 
-			for (FMediaFrameworkCaptureRenderTargetCameraOutputInfo& RenderTargetOutputInfo : MediaCaptureSettings->RenderTargetCaptures)
-			{
-				if (RenderTargetOutputInfo.MediaOutput == MediaOutput)
+			ForEachOutputInfoForMediaOutput<FMediaFrameworkCaptureRenderTargetCameraOutputInfo>(MediaCaptureSettings, MediaOutput,
+				[&MediaCaptureSettings, &CaptureSettingStructs, &bHasRenderTargetCaptures](FMediaFrameworkCaptureRenderTargetCameraOutputInfo& OutputInfo)
 				{
-					TSharedPtr<FStructOnScope> Struct = MakeShared<FStructOnScope>(FMediaFrameworkCaptureRenderTargetCameraOutputInfo::StaticStruct(), reinterpret_cast<uint8*>(&RenderTargetOutputInfo));
+						TSharedPtr<FStructOnScope> Struct = MakeShared<FStructOnScope>(FMediaFrameworkCaptureRenderTargetCameraOutputInfo::StaticStruct(), reinterpret_cast<uint8*>(&OutputInfo));
+						Struct->SetPackage(MediaCaptureSettings->GetPackage());
+						
+						CaptureSettingStructs.Add(Struct);
+						bHasRenderTargetCaptures = true;
+				});
+
+			ForEachOutputInfoForMediaOutput<FMediaFrameworkCaptureMediaTextureOutputInfo>(MediaCaptureSettings, MediaOutput,
+				[&MediaCaptureSettings, &CaptureSettingStructs, &bHasMediaTextureCaptures](FMediaFrameworkCaptureMediaTextureOutputInfo& OutputInfo)
+				{
+					TSharedPtr<FStructOnScope> Struct = MakeShared<FStructOnScope>(FMediaFrameworkCaptureMediaTextureOutputInfo::StaticStruct(), reinterpret_cast<uint8*>(&OutputInfo));
 					Struct->SetPackage(MediaCaptureSettings->GetPackage());
 					
 					CaptureSettingStructs.Add(Struct);
-					bHasRenderTargetCaptures = true;
-				}
-			}
+					bHasMediaTextureCaptures = true;
+				});
 		}
 	}
 
 	// Only show the capture settings if all objects have capture settings of the same type (current viewport, media viewport, or render target)
-	const bool bHasMultipleValues = (bHasCurrentViewportCapture && (bHasViewportCaptures || bHasRenderTargetCaptures)) || (bHasViewportCaptures && bHasRenderTargetCaptures);
-	if (CaptureSettingStructs.Num() && !bHasMultipleValues)
+	const int32 NumTrueValues = (int32)bHasCurrentViewportCapture + (int32)bHasViewportCaptures + (int32)bHasRenderTargetCaptures + (int32)bHasMediaTextureCaptures;
+	if (CaptureSettingStructs.Num() && NumTrueValues <= 1)
 	{
 		FAddPropertyParams AddPropertyParams;
 		AddPropertyParams.HideRootObjectNode(true);
@@ -1149,33 +1356,17 @@ void FMediaProfileMediaOutputDetailCustomization::SetMediaObject(UMediaOutput* I
 	}
 
 	bool bHasExistingCaptureInfo = false;
-	if (CaptureSettings->CurrentViewportMediaOutput.MediaOutput == InOriginalMediaObject)
-	{
-		CaptureSettings->CurrentViewportMediaOutput.MediaOutput = InNewMediaObject;
-		bHasExistingCaptureInfo = true;
-	}
-	
-	for (FMediaFrameworkCaptureCameraViewportCameraOutputInfo& ViewportOutputInfo : CaptureSettings->ViewportCaptures)
+	auto UpdateMediaOutputRef = [&InNewMediaObject, &bHasExistingCaptureInfo]<typename TOutputInfo>(TOutputInfo& OutputInfo)
 	{
-		if (ViewportOutputInfo.MediaOutput != InOriginalMediaObject)
-		{
-			continue;
-		}
-
-		ViewportOutputInfo.MediaOutput = InNewMediaObject;
+		OutputInfo.MediaOutput = InNewMediaObject;
 		bHasExistingCaptureInfo = true;
-	}
-
-	for (FMediaFrameworkCaptureRenderTargetCameraOutputInfo& RenderTargetOutputInfo : CaptureSettings->RenderTargetCaptures)
-	{
-		if (RenderTargetOutputInfo.MediaOutput != InOriginalMediaObject)
-		{
-			continue;
-		}
+	};
 
-		RenderTargetOutputInfo.MediaOutput = InNewMediaObject;
-		bHasExistingCaptureInfo = true;
-	}
+	using namespace UE::MediaFrameworkWorldSettings::Helpers;
+	ForEachOutputInfoForMediaOutput<FMediaFrameworkCaptureCurrentViewportOutputInfo>(CaptureSettings, InOriginalMediaObject, UpdateMediaOutputRef);
+	ForEachOutputInfoForMediaOutput<FMediaFrameworkCaptureCameraViewportCameraOutputInfo>(CaptureSettings, InOriginalMediaObject, UpdateMediaOutputRef);
+	ForEachOutputInfoForMediaOutput<FMediaFrameworkCaptureRenderTargetCameraOutputInfo>(CaptureSettings, InOriginalMediaObject, UpdateMediaOutputRef);
+	ForEachOutputInfoForMediaOutput<FMediaFrameworkCaptureMediaTextureOutputInfo>(CaptureSettings, InOriginalMediaObject, UpdateMediaOutputRef);
 
 	if (!bHasExistingCaptureInfo)
 	{
@@ -1259,7 +1450,8 @@ void FMediaProfileMediaOutputDetailCustomization::OnObjectPropertyChanged(UObjec
 		// as the capture properties may now be invalid
 		if (PropertyName == TEXT("MediaOutput") ||
 			PropertyName == GET_MEMBER_NAME_CHECKED(UMediaFrameworkWorldSettingsAssetUserData, ViewportCaptures) ||
-			PropertyName == GET_MEMBER_NAME_CHECKED(UMediaFrameworkWorldSettingsAssetUserData, RenderTargetCaptures))
+			PropertyName == GET_MEMBER_NAME_CHECKED(UMediaFrameworkWorldSettingsAssetUserData, RenderTargetCaptures) ||
+			PropertyName == GET_MEMBER_NAME_CHECKED(UMediaFrameworkWorldSettingsAssetUserData, MediaTextureCaptures))
 		{
 			CachedDetailBuilder.Pin()->ForceRefreshDetails();
 		}
diff --git a/Engine/Plugins/Media/MediaFrameworkUtilities/Source/MediaFrameworkUtilitiesEditor/Private/AssetEditor/SMediaProfileDetailsInfoPanel.cpp b/Engine/Plugins/Media/MediaFrameworkUtilities/Source/MediaFrameworkUtilitiesEditor/Private/AssetEditor/SMediaProfileDetailsInfoPanel.cpp
--- a/Engine/Plugins/Media/MediaFrameworkUtilities/Source/MediaFrameworkUtilitiesEditor/Private/AssetEditor/SMediaProfileDetailsInfoPanel.cpp
+++ b/Engine/Plugins/Media/MediaFrameworkUtilities/Source/MediaFrameworkUtilitiesEditor/Private/AssetEditor/SMediaProfileDetailsInfoPanel.cpp
@@ -9,6 +9,7 @@
 #include "MediaProfileEditor.h"
 #include "MediaProfileEditorUserSettings.h"
 #include "MediaSource.h"
+#include "MediaTexture.h"
 #include "Engine/TextureRenderTarget2D.h"
 #include "MediaAssets/ProxyMediaOutput.h"
 #include "MediaAssets/ProxyMediaSource.h"
@@ -325,19 +326,15 @@ FText SMediaObjectInfoPanel::GetCaptureMethodText() const
 		return FText::GetEmpty();
 	}
 
-	const bool bHasCurrentViewportCapture = CaptureSettings->CurrentViewportMediaOutput.MediaOutput == MediaOutput;
-	const bool bHasViewportCaptures = CaptureSettings->ViewportCaptures.ContainsByPredicate([MediaOutput](const FMediaFrameworkCaptureCameraViewportCameraOutputInfo& Info)
-	{
-		return Info.MediaOutput == MediaOutput;
-	});
+	using namespace UE::MediaFrameworkWorldSettings::Helpers;
 
-	const bool bHasRenderTargetCaptures = CaptureSettings->RenderTargetCaptures.ContainsByPredicate([MediaOutput](const FMediaFrameworkCaptureRenderTargetCameraOutputInfo& Info)
-	{
-		return Info.MediaOutput == MediaOutput;
-	});
+	const int32 bHasCurrentViewportCapture = !FindAllOutputInfoForMediaOutput<FMediaFrameworkCaptureCurrentViewportOutputInfo>(CaptureSettings, MediaOutput).IsEmpty();
+	const int32 bHasViewportCaptures = !FindAllOutputInfoForMediaOutput<FMediaFrameworkCaptureCameraViewportCameraOutputInfo>(CaptureSettings, MediaOutput).IsEmpty();
+	const int32 bHasRenderTargetCaptures = !FindAllOutputInfoForMediaOutput<FMediaFrameworkCaptureRenderTargetCameraOutputInfo>(CaptureSettings, MediaOutput).IsEmpty();
+	const int32 bHasMediaTextureCaptures = !FindAllOutputInfoForMediaOutput<FMediaFrameworkCaptureMediaTextureOutputInfo>(CaptureSettings, MediaOutput).IsEmpty();
 
-	const bool bHasMultipleValues = (bHasCurrentViewportCapture && (bHasViewportCaptures || bHasRenderTargetCaptures)) || (bHasViewportCaptures && bHasRenderTargetCaptures);
-	if (bHasMultipleValues)
+	const int32 NumTrueValues = (int32)bHasCurrentViewportCapture + (int32)bHasViewportCaptures + (int32)bHasRenderTargetCaptures + (int32)bHasMediaTextureCaptures;
+	if (NumTrueValues >= 2)
 	{
 		return LOCTEXT("MultipleCapturesValue", "Multiple Captures");
 	}
@@ -357,6 +354,11 @@ FText SMediaObjectInfoPanel::GetCaptureMethodText() const
 		return LOCTEXT("RenderTargetCaptureValue", "Render Target");
 	}
 
+	if (bHasMediaTextureCaptures)
+	{
+		return LOCTEXT("MediaTextureCaptureValue", "Media Texture");
+	}
+
 	return LOCTEXT("NoCapturesValues", "No Captures Configured");
 }
 
@@ -369,19 +371,15 @@ FText SMediaObjectInfoPanel::GetCaptureObjectLabelText() const
 		return FText::GetEmpty();
 	}
 
-	const bool bHasCurrentViewportCapture = CaptureSettings->CurrentViewportMediaOutput.MediaOutput == MediaOutput;
-	bool bHasViewportCaptures = CaptureSettings->ViewportCaptures.ContainsByPredicate([MediaOutput](const FMediaFrameworkCaptureCameraViewportCameraOutputInfo& Info)
-	{
-		return Info.MediaOutput == MediaOutput;
-	});
+	using namespace UE::MediaFrameworkWorldSettings::Helpers;
 
-	bool bHasRenderTargetCaptures = CaptureSettings->RenderTargetCaptures.ContainsByPredicate([MediaOutput](const FMediaFrameworkCaptureRenderTargetCameraOutputInfo& Info)
-	{
-		return Info.MediaOutput == MediaOutput;
-	});
+	const int32 bHasCurrentViewportCapture = !FindAllOutputInfoForMediaOutput<FMediaFrameworkCaptureCurrentViewportOutputInfo>(CaptureSettings, MediaOutput).IsEmpty();
+	const int32 bHasViewportCaptures = !FindAllOutputInfoForMediaOutput<FMediaFrameworkCaptureCameraViewportCameraOutputInfo>(CaptureSettings, MediaOutput).IsEmpty();
+	const int32 bHasRenderTargetCaptures = !FindAllOutputInfoForMediaOutput<FMediaFrameworkCaptureRenderTargetCameraOutputInfo>(CaptureSettings, MediaOutput).IsEmpty();
+	const int32 bHasMediaTextureCaptures = !FindAllOutputInfoForMediaOutput<FMediaFrameworkCaptureMediaTextureOutputInfo>(CaptureSettings, MediaOutput).IsEmpty();
 
-	const bool bHasMultipleValues = (bHasCurrentViewportCapture && (bHasViewportCaptures || bHasRenderTargetCaptures)) || (bHasViewportCaptures && bHasRenderTargetCaptures);
-	if (bHasMultipleValues)
+	const int32 NumTrueValues = (int32)bHasCurrentViewportCapture + (int32)bHasViewportCaptures + (int32)bHasRenderTargetCaptures + (int32)bHasMediaTextureCaptures;
+	if (NumTrueValues >= 2)
 	{
 		return LOCTEXT("MultipleCaptureObjectsLabel", "Multiple Objects:");
 	}
@@ -401,6 +399,11 @@ FText SMediaObjectInfoPanel::GetCaptureObjectLabelText() const
 		return LOCTEXT("RenderTargetObjectLabel", "Render Target:");
 	}
 
+	if (bHasMediaTextureCaptures)
+	{
+		return LOCTEXT("MediaTextureObjectLabel", "Media Texture:");
+	}
+
 	return LOCTEXT("NoObjectsLabel", "No Objects:");
 }
 
@@ -413,17 +416,15 @@ FText SMediaObjectInfoPanel::GetCaptureObjectValueText() const
 		return FText::GetEmpty();
 	}
 
-	const bool bHasCurrentViewportCapture = CaptureSettings->CurrentViewportMediaOutput.MediaOutput == MediaOutput;
+	using namespace UE::MediaFrameworkWorldSettings::Helpers;
+
+	const bool bHasCurrentViewportCapture = !FindAllOutputInfoForMediaOutput<FMediaFrameworkCaptureCurrentViewportOutputInfo>(CaptureSettings, MediaOutput).IsEmpty();
 
 	bool bHasViewportCaptures = false;
 	FText CamerasText;
-	for (const FMediaFrameworkCaptureCameraViewportCameraOutputInfo& Info : CaptureSettings->ViewportCaptures)
+	ForEachOutputInfoForMediaOutput<FMediaFrameworkCaptureCameraViewportCameraOutputInfo>(CaptureSettings, MediaOutput, 
+	[&bHasViewportCaptures, &CamerasText](const FMediaFrameworkCaptureCameraViewportCameraOutputInfo& Info)
 	{
-		if (Info.MediaOutput != MediaOutput)
-		{
-			continue;
-		}
-
 		bHasViewportCaptures = true;
 		
 		for (const TSoftObjectPtr<AActor>& Camera : Info.Cameras)
@@ -441,20 +442,15 @@ FText SMediaObjectInfoPanel::GetCaptureObjectValueText() const
 				}
 			}
 		}
-	}
+	});
 
 	bool bHasRenderTargetCaptures = false;
 	FText RenderTargetsText;
-	for (const FMediaFrameworkCaptureRenderTargetCameraOutputInfo& Info : CaptureSettings->RenderTargetCaptures)
+	ForEachOutputInfoForMediaOutput<FMediaFrameworkCaptureRenderTargetCameraOutputInfo>(CaptureSettings, MediaOutput, 
+	[&bHasRenderTargetCaptures, &RenderTargetsText](const FMediaFrameworkCaptureRenderTargetCameraOutputInfo& Info)
 	{
-		if (Info.MediaOutput != MediaOutput)
-		{
-			continue;
-		}
-
 		bHasRenderTargetCaptures = true;
 		
-		
 		FText RenderTargetText = Info.RenderTarget ? FText::FromString(Info.RenderTarget->GetName()) : LOCTEXT("NoRenderTarget", "No Render Target");
 		
 		if (RenderTargetsText.IsEmpty())
@@ -465,10 +461,29 @@ FText SMediaObjectInfoPanel::GetCaptureObjectValueText() const
 		{
 			RenderTargetsText = FText::Format(LOCTEXT("AppendRenderTargetNameFormat", "{0}, {1}"), RenderTargetsText, RenderTargetText);
 		}
-	}
+	});
 
-	const bool bHasMultipleValues = (bHasCurrentViewportCapture && (bHasViewportCaptures || bHasRenderTargetCaptures)) || (bHasViewportCaptures && bHasRenderTargetCaptures);
-	if (bHasMultipleValues)
+	bool bHasMediaTextureCaptures = false;
+	FText MediaTexturesText;
+	ForEachOutputInfoForMediaOutput<FMediaFrameworkCaptureMediaTextureOutputInfo>(CaptureSettings, MediaOutput, 
+	[&bHasMediaTextureCaptures, &MediaTexturesText](const FMediaFrameworkCaptureMediaTextureOutputInfo& Info)
+	{
+		bHasMediaTextureCaptures = true;
+		
+		FText MediaTextureText = Info.MediaTexture ? FText::FromString(Info.MediaTexture->GetName()) : LOCTEXT("NoMediaTexture", "No Media Texture");
+		
+		if (MediaTexturesText.IsEmpty())
+		{
+			MediaTexturesText = MediaTextureText;
+		}
+		else
+		{
+			MediaTexturesText = FText::Format(LOCTEXT("AppendMediaTextureNameFormat", "{0}, {1}"), MediaTexturesText, MediaTextureText);
+		}
+	});
+	
+	const int32 NumTrueValues = (int32)bHasCurrentViewportCapture + (int32)bHasViewportCaptures + (int32)bHasRenderTargetCaptures + (int32)bHasMediaTextureCaptures;
+	if (NumTrueValues >= 2)
 	{
 		return FText::GetEmpty();
 	}
@@ -483,6 +498,11 @@ FText SMediaObjectInfoPanel::GetCaptureObjectValueText() const
 		return RenderTargetsText;
 	}
 
+	if (bHasMediaTextureCaptures)
+	{
+		return MediaTexturesText;
+	}
+
 	return FText::GetEmpty();
 }
 
@@ -518,54 +538,26 @@ FText SMediaObjectInfoPanel::GetCaptureColorConversionText() const
 	}
 
 	FText ColorConversionText;
-	if (CaptureSettings->CurrentViewportMediaOutput.MediaOutput == MediaOutput)
+	auto AppendColorConversionText = [&ColorConversionText]<typename TOutputInfo>(const TOutputInfo& Info)
 	{
 		if (ColorConversionText.IsEmpty())
 		{
-			ColorConversionText = FText::FromString(CaptureSettings->CurrentViewportMediaOutput.CaptureOptions.GetColorConversionSettingsString());
+			ColorConversionText = FText::FromString(Info.CaptureOptions.GetColorConversionSettingsString());
 		}
 		else
 		{
 			ColorConversionText = FText::Format(
 				LOCTEXT("AppendColorConversionFormat", "{0}, {1}"),
 				ColorConversionText,
-				FText::FromString(CaptureSettings->CurrentViewportMediaOutput.CaptureOptions.GetColorConversionSettingsString()));
-		}
-	}
-	
-	for (const FMediaFrameworkCaptureCameraViewportCameraOutputInfo& Info : CaptureSettings->ViewportCaptures)
-	{
-		if (Info.MediaOutput != MediaOutput)
-		{
-			continue;
+				FText::FromString(Info.CaptureOptions.GetColorConversionSettingsString()));
 		}
+	};
 
-		if (ColorConversionText.IsEmpty())
-		{
-			ColorConversionText = FText::FromString(Info.CaptureOptions.GetColorConversionSettingsString());
-		}
-		else
-		{
-			ColorConversionText = FText::Format(LOCTEXT("AppendColorConversionFormat", "{0}, {1}"), ColorConversionText, FText::FromString(Info.CaptureOptions.GetColorConversionSettingsString()));
-		}
-	}
-	
-	for (const FMediaFrameworkCaptureRenderTargetCameraOutputInfo& Info : CaptureSettings->RenderTargetCaptures)
-	{
-		if (Info.MediaOutput != MediaOutput)
-		{
-			continue;
-		}
-
-		if (ColorConversionText.IsEmpty())
-		{
-			ColorConversionText = FText::FromString(Info.CaptureOptions.GetColorConversionSettingsString());
-		}
-		else
-		{
-			ColorConversionText = FText::Format(LOCTEXT("AppendColorConversionFormat", "{0}, {1}"), ColorConversionText, FText::FromString(Info.CaptureOptions.GetColorConversionSettingsString()));
-		}
-	}
+	using namespace UE::MediaFrameworkWorldSettings::Helpers;
+	ForEachOutputInfoForMediaOutput<FMediaFrameworkCaptureCurrentViewportOutputInfo>(CaptureSettings, MediaOutput, AppendColorConversionText);
+	ForEachOutputInfoForMediaOutput<FMediaFrameworkCaptureCameraViewportCameraOutputInfo>(CaptureSettings, MediaOutput, AppendColorConversionText);
+	ForEachOutputInfoForMediaOutput<FMediaFrameworkCaptureRenderTargetCameraOutputInfo>(CaptureSettings, MediaOutput, AppendColorConversionText);
+	ForEachOutputInfoForMediaOutput<FMediaFrameworkCaptureMediaTextureOutputInfo>(CaptureSettings, MediaOutput, AppendColorConversionText);
 
 	if (ColorConversionText.IsEmpty())
 	{
diff --git a/Engine/Plugins/Media/MediaFrameworkUtilities/Source/MediaFrameworkUtilitiesEditor/Private/AssetEditor/SMediaProfileDetailsPanel.cpp b/Engine/Plugins/Media/MediaFrameworkUtilities/Source/MediaFrameworkUtilitiesEditor/Private/AssetEditor/SMediaProfileDetailsPanel.cpp
--- a/Engine/Plugins/Media/MediaFrameworkUtilities/Source/MediaFrameworkUtilitiesEditor/Private/AssetEditor/SMediaProfileDetailsPanel.cpp
+++ b/Engine/Plugins/Media/MediaFrameworkUtilities/Source/MediaFrameworkUtilitiesEditor/Private/AssetEditor/SMediaProfileDetailsPanel.cpp
@@ -276,19 +276,17 @@ void SMediaProfileDetailsPanel::OnCaptureButtonPressed()
 		{
 			if (UMediaOutput* MediaOutput = PinnedMediaProfile->GetMediaOutput(OutputIndex))
 			{
-				if (MediaCaptureSettings->CurrentViewportMediaOutput.MediaOutput == MediaOutput)
+				using namespace UE::MediaFrameworkWorldSettings::Helpers;
+				
+				if (FMediaFrameworkCaptureCurrentViewportOutputInfo* CurrentViewportCapture =
+					FindFirstOutputInfoForMediaOutput<FMediaFrameworkCaptureCurrentViewportOutputInfo>(MediaCaptureSettings, MediaOutput))
 				{
-					PinnedMediaProfile->GetPlaybackManager()->OpenActiveViewportOutput(MediaOutput, MediaCaptureSettings->CurrentViewportMediaOutput.CaptureOptions, MediaCaptureSettings->bAutoRestartCaptureOnChange);
+					PinnedMediaProfile->GetPlaybackManager()->OpenActiveViewportOutput(MediaOutput, CurrentViewportCapture->CaptureOptions, MediaCaptureSettings->bAutoRestartCaptureOnChange);
 					continue;
 				}
 
-				FMediaFrameworkCaptureCameraViewportCameraOutputInfo* ViewportCapture =
-					MediaCaptureSettings->ViewportCaptures.FindByPredicate([MediaOutput](const FMediaFrameworkCaptureCameraViewportCameraOutputInfo& InViewportCapture)
-					{
-						return InViewportCapture.MediaOutput == MediaOutput;
-					});
-				
-				if (ViewportCapture)
+				if (FMediaFrameworkCaptureCameraViewportCameraOutputInfo* ViewportCapture =
+					FindFirstOutputInfoForMediaOutput<FMediaFrameworkCaptureCameraViewportCameraOutputInfo>(MediaCaptureSettings, MediaOutput))
 				{
 					// Make sure there is a managed viewport before starting the capture
 					TSharedPtr<FViewportClient> ViewportClient = PinnedMediaProfile->GetPlaybackManager()->GetOrCreateManagedViewport(MediaOutput, ViewportCapture->ViewMode);
@@ -299,16 +297,24 @@ void SMediaProfileDetailsPanel::OnCaptureButtonPressed()
 					continue;
 				}
 
-				FMediaFrameworkCaptureRenderTargetCameraOutputInfo* RenderTargetCapture =
-					MediaCaptureSettings->RenderTargetCaptures.FindByPredicate([MediaOutput](const FMediaFrameworkCaptureRenderTargetCameraOutputInfo& InRenderTargetCapture)
-					{
-						return InRenderTargetCapture.MediaOutput == MediaOutput;
-					});
-
-				if (RenderTargetCapture)
+				if (FMediaFrameworkCaptureRenderTargetCameraOutputInfo* RenderTargetCapture =
+					FindFirstOutputInfoForMediaOutput<FMediaFrameworkCaptureRenderTargetCameraOutputInfo>(MediaCaptureSettings, MediaOutput))
 				{
 					PinnedMediaProfile->GetPlaybackManager()->OpenRenderTargetOutput(MediaOutput, RenderTargetCapture->RenderTarget, RenderTargetCapture->CaptureOptions, MediaCaptureSettings->bAutoRestartCaptureOnChange);
 				}
+
+				if (FMediaFrameworkCaptureMediaTextureOutputInfo* MediaTextureCapture =
+					FindFirstOutputInfoForMediaOutput<FMediaFrameworkCaptureMediaTextureOutputInfo>(MediaCaptureSettings, MediaOutput))
+				{
+					// For media texture capture, if the media texture is from a media profile source, make sure the source is open before starting capture
+					int32 MediaSourceIndex = INDEX_NONE;
+					if (PinnedMediaProfile->GetPlaybackManager()->IsValidSourceMediaTexture(MediaTextureCapture->MediaTexture, MediaSourceIndex))
+					{
+						PinnedMediaProfile->GetPlaybackManager()->OpenSourceFromIndex(MediaSourceIndex, OwningEditor.Pin().Get());
+					}
+					
+					PinnedMediaProfile->GetPlaybackManager()->OpenMediaTextureOutput(MediaOutput, MediaTextureCapture->MediaTexture, MediaTextureCapture->CaptureOptions, MediaCaptureSettings->bAutoRestartCaptureOnChange);
+				}
 			}
 		}
 	}
diff --git a/Engine/Plugins/Media/MediaFrameworkUtilities/Source/MediaFrameworkUtilitiesEditor/Private/AssetEditor/SMediaProfileMediaItemDisplay.cpp b/Engine/Plugins/Media/MediaFrameworkUtilities/Source/MediaFrameworkUtilitiesEditor/Private/AssetEditor/SMediaProfileMediaItemDisplay.cpp
--- a/Engine/Plugins/Media/MediaFrameworkUtilities/Source/MediaFrameworkUtilitiesEditor/Private/AssetEditor/SMediaProfileMediaItemDisplay.cpp
+++ b/Engine/Plugins/Media/MediaFrameworkUtilities/Source/MediaFrameworkUtilitiesEditor/Private/AssetEditor/SMediaProfileMediaItemDisplay.cpp
@@ -844,6 +844,92 @@ private:
 	TAttribute<bool> DrawCheckerboardAttr;
 };
 
+/** Widget to display a media texture media capture output */
+class SMediaCaptureMediaTextureImage : public SMediaCaptureImageBase
+{
+public:
+	SLATE_BEGIN_ARGS(SMediaCaptureMediaTextureImage) {}
+		SLATE_ATTRIBUTE(TOptional<EMediaCaptureState>, MediaCaptureState)
+		SLATE_ATTRIBUTE(FOptionalSize, AspectRatio)
+		SLATE_ATTRIBUTE(EColorChannelMask, ChannelMask)
+		SLATE_ATTRIBUTE(bool, InvertAlphaChannel)
+		SLATE_ATTRIBUTE(bool, DrawCheckerboard)
+	SLATE_END_ARGS()
+
+	void Construct(const FArguments& InArgs, UMediaTexture* InMediaTexture)
+	{
+		ChannelMaskAttr = InArgs._ChannelMask;
+		InvertAlphaChannelAttr = InArgs._InvertAlphaChannel;
+		DrawCheckerboardAttr = InArgs._DrawCheckerboard;
+		
+		SMediaCaptureImageBase::FArguments BaseArgs;
+		BaseArgs.MediaCaptureState(InArgs._MediaCaptureState);
+		BaseArgs.AspectRatio(InArgs._AspectRatio);
+		
+		SMediaCaptureImageBase::Construct(BaseArgs);
+
+		Refresh(InMediaTexture);
+	}
+
+	/** Updates the display if the capture's media texture has been changed */
+	void Refresh(UMediaTexture* InMediaTexture)
+	{
+		if (MediaTexture.IsValid() && MediaTexture.Get() == InMediaTexture)
+		{
+			return;
+		}
+
+		MediaTexture = InMediaTexture;
+		if (InMediaTexture)
+		{
+			Container->SetContent
+			(
+				SNew(SMaskedMediaImage, InMediaTexture)
+				.IsAlphaPremultiplied(true)
+				.ChannelMask(ChannelMaskAttr)
+				.InvertAlphaChannel(InvertAlphaChannelAttr)
+				.DrawCheckerboard(DrawCheckerboardAttr)
+			);
+		}
+		else
+		{
+			Container->SetContent
+			(
+				SNew(SBorder)
+				.HAlign(HAlign_Center)
+				.VAlign(VAlign_Center)
+				.BorderImage(FAppStyle::GetBrush("NoBorder"))
+				[
+					SNew(STextBlock)
+					.Text(LOCTEXT("NoMediaTextureLabel", "No Media Texture set for Media Output Capture"))
+				]
+			);
+		}
+	}
+
+protected:
+	virtual TSharedRef<SWidget> GetContentWidget() override
+	{
+		return SAssignNew(Container, SBox);
+	}
+	
+private:
+	/** Widget that contains the render target output widget */
+	TSharedPtr<SBox> Container;
+
+	/** The media texture being used by the capture */
+	TWeakObjectPtr<UMediaTexture> MediaTexture;
+	
+	/** Attribute for the color channel to mask to */
+	TAttribute<EColorChannelMask> ChannelMaskAttr;
+
+	/** Attribute for the invert alpha flag */
+	TAttribute<bool> InvertAlphaChannelAttr;
+
+	/** Attribute for drawing the alpha-blended checkerboard background */
+	TAttribute<bool> DrawCheckerboardAttr;
+};
+
 void SMediaProfileMediaOutputDisplay::Construct(const FArguments& InArgs, const TSharedPtr<SMediaProfileViewport>& InOwningViewport)
 {
 	ShowFlagsCommandList = MakeShared<FUICommandList>();
@@ -871,8 +957,8 @@ SMediaProfileMediaOutputDisplay::~SMediaProfileMediaOutputDisplay()
 void SMediaProfileMediaOutputDisplay::ConfigureMediaImage()
 {
 	ImageWidget.Reset();
-	ViewportOutputInfoIndex = INDEX_NONE;
-	RenderTargetOutputInfoIndex = INDEX_NONE;
+	CaptureType = ECaptureType::None;
+	CaptureIndex = INDEX_NONE;
 
 	// Clear out all commands that have been bound to any viewport clients from viewport captures
 	for (const FShowFlagMenuCommands::FShowFlagCommand& FlagCommand : FShowFlagMenuCommands::Get().GetCommands())
@@ -906,6 +992,8 @@ void SMediaProfileMediaOutputDisplay::ConfigureMediaImage()
 	{
 		if (CaptureSettings->CurrentViewportMediaOutput.MediaOutput == MediaOutput)
 		{
+			CaptureType = ECaptureType::ActiveViewport;
+
 			TSharedPtr<SMediaCaptureEditorViewportImage> CurrentViewportImage;
 			MediaImageContainer->SetContent(
 				SAssignNew(CurrentViewportImage, SMediaCaptureEditorViewportImage, MediaProfile, CaptureSettings->CurrentViewportMediaOutput)
@@ -925,17 +1013,19 @@ void SMediaProfileMediaOutputDisplay::ConfigureMediaImage()
 			return;
 		}
 		
-		ViewportOutputInfoIndex = CaptureSettings->ViewportCaptures.IndexOfByPredicate(
+		CaptureIndex = CaptureSettings->ViewportCaptures.IndexOfByPredicate(
 		[MediaOutput](const FMediaFrameworkCaptureCameraViewportCameraOutputInfo& OutputInfo)
 			{
 				return OutputInfo.MediaOutput == MediaOutput;
 			});
 		
-		if (ViewportOutputInfoIndex != INDEX_NONE)
+		if (CaptureSettings->ViewportCaptures.IsValidIndex(CaptureIndex))
 		{
+			CaptureType = ECaptureType::Viewport;
+
 			TSharedPtr<SMediaCaptureViewportImage> ViewportImage;
 			MediaImageContainer->SetContent(
-				SAssignNew(ViewportImage, SMediaCaptureViewportImage, MediaProfile, CaptureSettings->ViewportCaptures[ViewportOutputInfoIndex])
+				SAssignNew(ViewportImage, SMediaCaptureViewportImage, MediaProfile, CaptureSettings->ViewportCaptures[CaptureIndex])
 				.MediaCaptureState(this, &SMediaProfileMediaOutputDisplay::GetMediaOutputCaptureState)
 				.AspectRatio(this, &SMediaProfileMediaOutputDisplay::GetMediaOutputDesiredAspectRatio)
 			);
@@ -963,16 +1053,39 @@ void SMediaProfileMediaOutputDisplay::ConfigureMediaImage()
 			return;
 		}
 
-		RenderTargetOutputInfoIndex = CaptureSettings->RenderTargetCaptures.IndexOfByPredicate(
+		CaptureIndex = CaptureSettings->RenderTargetCaptures.IndexOfByPredicate(
 		[MediaOutput](const FMediaFrameworkCaptureRenderTargetCameraOutputInfo& OutputInfo)
 			{
 				return OutputInfo.MediaOutput == MediaOutput;
 			});
 
-		if (RenderTargetOutputInfoIndex != INDEX_NONE)
+		if (CaptureSettings->RenderTargetCaptures.IsValidIndex(CaptureIndex))
 		{
+			CaptureType = ECaptureType::RenderTarget;
+
 			MediaImageContainer->SetContent(
-				SAssignNew(ImageWidget, SMediaCaptureRenderTargetImage, CaptureSettings->RenderTargetCaptures[RenderTargetOutputInfoIndex].RenderTarget)
+				SAssignNew(ImageWidget, SMediaCaptureRenderTargetImage, CaptureSettings->RenderTargetCaptures[CaptureIndex].RenderTarget)
+				.MediaCaptureState(this, &SMediaProfileMediaOutputDisplay::GetMediaOutputCaptureState)
+				.AspectRatio(this, &SMediaProfileMediaOutputDisplay::GetMediaOutputDesiredAspectRatio)
+				.ChannelMask(this, &SMediaProfileMediaOutputDisplay::GetChannelMask)
+				.InvertAlphaChannel(this, &SMediaProfileMediaOutputDisplay::GetInvertAlphaChannelMask)
+				.DrawCheckerboard(this, &SMediaProfileMediaOutputDisplay::GetDrawAlphaBlendedCheckerboard)
+			);
+			return;
+		}
+
+		CaptureIndex = CaptureSettings->MediaTextureCaptures.IndexOfByPredicate(
+			[MediaOutput](const FMediaFrameworkCaptureMediaTextureOutputInfo& OutputInfo)
+			{
+				return OutputInfo.MediaOutput == MediaOutput;
+			});
+
+		if (CaptureSettings->MediaTextureCaptures.IsValidIndex(CaptureIndex))
+		{
+			CaptureType = ECaptureType::MediaTexture;
+
+			MediaImageContainer->SetContent(
+				SAssignNew(ImageWidget, SMediaCaptureMediaTextureImage, CaptureSettings->MediaTextureCaptures[CaptureIndex].MediaTexture)
 				.MediaCaptureState(this, &SMediaProfileMediaOutputDisplay::GetMediaOutputCaptureState)
 				.AspectRatio(this, &SMediaProfileMediaOutputDisplay::GetMediaOutputDesiredAspectRatio)
 				.ChannelMask(this, &SMediaProfileMediaOutputDisplay::GetChannelMask)
@@ -1092,7 +1205,7 @@ void SMediaProfileMediaOutputDisplay::OnObjectPropertyChanged(UObject* InObject,
 		const FName PropertyName = InPropertyChangedEvent.GetPropertyName();
 		if (PropertyName == GET_MEMBER_NAME_CHECKED(FMediaFrameworkCaptureRenderTargetCameraOutputInfo, RenderTarget))
 		{
-			if (RenderTargetOutputInfoIndex == INDEX_NONE || !CaptureSettings->RenderTargetCaptures.IsValidIndex(RenderTargetOutputInfoIndex))
+			if (CaptureType != ECaptureType::RenderTarget || !CaptureSettings->RenderTargetCaptures.IsValidIndex(CaptureIndex))
 			{
 				return;
 			}
@@ -1103,13 +1216,13 @@ void SMediaProfileMediaOutputDisplay::OnObjectPropertyChanged(UObject* InObject,
 				return;
 			}
 
-			RenderTargetImageWidget->Refresh(CaptureSettings->RenderTargetCaptures[RenderTargetOutputInfoIndex].RenderTarget);
+			RenderTargetImageWidget->Refresh(CaptureSettings->RenderTargetCaptures[CaptureIndex].RenderTarget);
 			return;
 		}
 
 		if (PropertyName == GET_MEMBER_NAME_CHECKED(FMediaFrameworkCaptureCameraViewportCameraOutputInfo, Cameras))
 		{
-			if (ViewportOutputInfoIndex == INDEX_NONE || !CaptureSettings->ViewportCaptures.IsValidIndex(ViewportOutputInfoIndex))
+			if (CaptureType != ECaptureType::Viewport || !CaptureSettings->ViewportCaptures.IsValidIndex(CaptureIndex))
 			{
 				return;
 			}
@@ -1120,32 +1233,51 @@ void SMediaProfileMediaOutputDisplay::OnObjectPropertyChanged(UObject* InObject,
 				return;
 			}
 
-			ViewportImageWidget->Refresh(CaptureSettings->ViewportCaptures[ViewportOutputInfoIndex].Cameras);
+			ViewportImageWidget->Refresh(CaptureSettings->ViewportCaptures[CaptureIndex].Cameras);
 			return;
 		}
 
+		if (PropertyName == GET_MEMBER_NAME_CHECKED(FMediaFrameworkCaptureMediaTextureOutputInfo, MediaTexture))
+		{
+			if (CaptureType != ECaptureType::MediaTexture || !CaptureSettings->MediaTextureCaptures.IsValidIndex(CaptureIndex))
+			{
+				return;
+			}
+	
+			TSharedPtr<SMediaCaptureMediaTextureImage> MediaTextureImageWidget = StaticCastSharedPtr<SMediaCaptureMediaTextureImage>(ImageWidget);
+			if (!MediaTextureImageWidget.IsValid())
+			{
+				return;
+			}
+
+			MediaTextureImageWidget->Refresh(CaptureSettings->MediaTextureCaptures[CaptureIndex].MediaTexture);
+			return;
+		}
+		
 		// If a media output property in the capture settings has changed, or the viewport captures or render target captures lists have changed,
 		// check to see the media output object for this display is still linked to a valid capture. If not, refresh the widget
 		if (PropertyName == TEXT("MediaOutput") ||
 			PropertyName == GET_MEMBER_NAME_CHECKED(UMediaFrameworkWorldSettingsAssetUserData, ViewportCaptures) ||
-			PropertyName == GET_MEMBER_NAME_CHECKED(UMediaFrameworkWorldSettingsAssetUserData, RenderTargetCaptures))
+			PropertyName == GET_MEMBER_NAME_CHECKED(UMediaFrameworkWorldSettingsAssetUserData, RenderTargetCaptures) ||
+			PropertyName == GET_MEMBER_NAME_CHECKED(UMediaFrameworkWorldSettingsAssetUserData, MediaTextureCaptures))
 		{
-			const bool bIsViewportCapture = ViewportOutputInfoIndex != INDEX_NONE;
-			const bool bIsRenderTargetCapture = RenderTargetOutputInfoIndex != INDEX_NONE;
-			const bool bIsActiveEditorViewportCapture = ViewportOutputInfoIndex == INDEX_NONE && RenderTargetOutputInfoIndex == INDEX_NONE;
-
 			bool bResetMediaImage = false;
-			if (bIsViewportCapture)
+			if (CaptureType == ECaptureType::Viewport)
+			{
+				bResetMediaImage = !CaptureSettings->ViewportCaptures.IsValidIndex(CaptureIndex) ||
+					CaptureSettings->ViewportCaptures[CaptureIndex].MediaOutput != GetMediaItem();
+			}
+			else if (CaptureType == ECaptureType::RenderTarget)
 			{
-				bResetMediaImage = !CaptureSettings->ViewportCaptures.IsValidIndex(ViewportOutputInfoIndex) ||
-					CaptureSettings->ViewportCaptures[ViewportOutputInfoIndex].MediaOutput != GetMediaItem();
+				bResetMediaImage = !CaptureSettings->RenderTargetCaptures.IsValidIndex(CaptureIndex) ||
+					CaptureSettings->RenderTargetCaptures[CaptureIndex].MediaOutput != GetMediaItem();
 			}
-			else if (bIsRenderTargetCapture)
+			else if (CaptureType == ECaptureType::MediaTexture)
 			{
-				bResetMediaImage = !CaptureSettings->RenderTargetCaptures.IsValidIndex(RenderTargetOutputInfoIndex) ||
-					CaptureSettings->RenderTargetCaptures[RenderTargetOutputInfoIndex].MediaOutput != GetMediaItem();
+				bResetMediaImage = !CaptureSettings->MediaTextureCaptures.IsValidIndex(CaptureIndex) ||
+					CaptureSettings->MediaTextureCaptures[CaptureIndex].MediaOutput != GetMediaItem();
 			}
-			else if (bIsActiveEditorViewportCapture)
+			else if (CaptureType == ECaptureType::ActiveViewport)
 			{
 				bResetMediaImage = CaptureSettings->CurrentViewportMediaOutput.MediaOutput != GetMediaItem();
 			}
diff --git a/Engine/Plugins/Media/MediaFrameworkUtilities/Source/MediaFrameworkUtilitiesEditor/Private/AssetEditor/SMediaProfileMediaItemDisplay.h b/Engine/Plugins/Media/MediaFrameworkUtilities/Source/MediaFrameworkUtilitiesEditor/Private/AssetEditor/SMediaProfileMediaItemDisplay.h
--- a/Engine/Plugins/Media/MediaFrameworkUtilities/Source/MediaFrameworkUtilitiesEditor/Private/AssetEditor/SMediaProfileMediaItemDisplay.h
+++ b/Engine/Plugins/Media/MediaFrameworkUtilities/Source/MediaFrameworkUtilitiesEditor/Private/AssetEditor/SMediaProfileMediaItemDisplay.h
@@ -1028,6 +1028,16 @@ private:
 /** Implementation of SMediaProfileMediaItemDisplayBase for media outputs, which displays the result from a media capture from the selected media output */
 class SMediaProfileMediaOutputDisplay : public SMediaProfileMediaItemDisplayBase<UMediaOutput>
 {
+private:
+	enum class ECaptureType : uint8
+	{
+		None,
+		ActiveViewport,
+		Viewport,
+		RenderTarget,
+		MediaTexture
+	};
+	
 public:
 	SLATE_BEGIN_ARGS(SMediaProfileMediaOutputDisplay) {}
 		SLATE_ARGUMENT(TWeakPtr<FMediaProfileEditor>, MediaProfileEditor)
@@ -1067,8 +1077,8 @@ private:
 private:
 	TSharedPtr<SWidget> ImageWidget;
 	TSharedPtr<FUICommandList> ShowFlagsCommandList;
-	int32 ViewportOutputInfoIndex = INDEX_NONE;
-	int32 RenderTargetOutputInfoIndex = INDEX_NONE;
+	ECaptureType CaptureType = ECaptureType::None;
+	int32 CaptureIndex = INDEX_NONE;
 	bool bDisplayShowFlags = false;
 };
 
diff --git a/Engine/Plugins/Media/MediaFrameworkUtilities/Source/MediaFrameworkUtilitiesEditor/Private/MediaFrameworkWorldSettingsAssetUserData.cpp b/Engine/Plugins/Media/MediaFrameworkUtilities/Source/MediaFrameworkUtilitiesEditor/Private/MediaFrameworkWorldSettingsAssetUserData.cpp
--- a/Engine/Plugins/Media/MediaFrameworkUtilities/Source/MediaFrameworkUtilitiesEditor/Private/MediaFrameworkWorldSettingsAssetUserData.cpp
+++ b/Engine/Plugins/Media/MediaFrameworkUtilities/Source/MediaFrameworkUtilitiesEditor/Private/MediaFrameworkWorldSettingsAssetUserData.cpp
@@ -50,6 +50,12 @@ FMediaFrameworkCaptureRenderTargetCameraOutputInfo::FMediaFrameworkCaptureRender
 {
 }
 
+FMediaFrameworkCaptureMediaTextureOutputInfo::FMediaFrameworkCaptureMediaTextureOutputInfo()
+	: MediaTexture(nullptr)
+	, MediaOutput(nullptr)
+{
+	CaptureOptions.bApplyLinearToSRGBConversion = true;
+}
 
 UMediaFrameworkWorldSettingsAssetUserData::UMediaFrameworkWorldSettingsAssetUserData()
 {
@@ -81,4 +87,270 @@ void UMediaFrameworkWorldSettingsAssetUserData::Serialize(FArchive& Ar)
 		}
 	}
 }
+
+namespace UE::MediaFrameworkWorldSettings::Helpers
+{
+	template <>
+	FMediaFrameworkCaptureCurrentViewportOutputInfo* FindFirstOutputInfoForMediaOutput<FMediaFrameworkCaptureCurrentViewportOutputInfo>(
+		UMediaFrameworkWorldSettingsAssetUserData* InSettings,
+		UMediaOutput* InMediaOutput)
+	{
+		if (InSettings != nullptr)
+		{
+			if (InSettings->CurrentViewportMediaOutput.MediaOutput == InMediaOutput)
+			{
+				return &InSettings->CurrentViewportMediaOutput;
+			}
+		}
+		
+		return nullptr;
+	}
+
+	template <>
+	FMediaFrameworkCaptureCameraViewportCameraOutputInfo* FindFirstOutputInfoForMediaOutput<FMediaFrameworkCaptureCameraViewportCameraOutputInfo>(
+		UMediaFrameworkWorldSettingsAssetUserData* InSettings,
+		UMediaOutput* InMediaOutput)
+	{
+		if (InSettings != nullptr)
+		{
+			for (FMediaFrameworkCaptureCameraViewportCameraOutputInfo& OutputInfo : InSettings->ViewportCaptures)
+			{
+				if (OutputInfo.MediaOutput == InMediaOutput)
+				{
+					return &OutputInfo;
+				}
+			}
+		}
+		
+		return nullptr;
+	}
+
+	template <>
+	FMediaFrameworkCaptureRenderTargetCameraOutputInfo* FindFirstOutputInfoForMediaOutput<FMediaFrameworkCaptureRenderTargetCameraOutputInfo>(
+		UMediaFrameworkWorldSettingsAssetUserData* InSettings,
+		UMediaOutput* InMediaOutput)
+	{
+		if (InSettings != nullptr)
+		{
+			for (FMediaFrameworkCaptureRenderTargetCameraOutputInfo& OutputInfo : InSettings->RenderTargetCaptures)
+			{
+				if (OutputInfo.MediaOutput == InMediaOutput)
+				{
+					return &OutputInfo;
+				}
+			}
+		}
+		
+		return nullptr;
+	}
+	
+	template <>
+	FMediaFrameworkCaptureMediaTextureOutputInfo* FindFirstOutputInfoForMediaOutput<FMediaFrameworkCaptureMediaTextureOutputInfo>(
+		UMediaFrameworkWorldSettingsAssetUserData* InSettings,
+		UMediaOutput* InMediaOutput)
+	{
+		if (InSettings != nullptr)
+		{
+			for (FMediaFrameworkCaptureMediaTextureOutputInfo& OutputInfo : InSettings->MediaTextureCaptures)
+			{
+				if (OutputInfo.MediaOutput == InMediaOutput)
+				{
+					return &OutputInfo;
+				}
+			}
+		}
+		
+		return nullptr;
+	}
+
+	template <>
+	TArray<FMediaFrameworkCaptureCurrentViewportOutputInfo>	FindAllOutputInfoForMediaOutput<FMediaFrameworkCaptureCurrentViewportOutputInfo>(
+		UMediaFrameworkWorldSettingsAssetUserData* InSettings,
+		UMediaOutput* InMediaOutput)
+	{
+		TArray<FMediaFrameworkCaptureCurrentViewportOutputInfo> OutputInfos;
+		if (InSettings != nullptr)
+		{
+			if (InSettings->CurrentViewportMediaOutput.MediaOutput == InMediaOutput)
+			{
+				OutputInfos.Add(InSettings->CurrentViewportMediaOutput);
+			}
+		}
+	
+		return OutputInfos;
+	}
+
+	template <>
+	TArray<FMediaFrameworkCaptureCameraViewportCameraOutputInfo> FindAllOutputInfoForMediaOutput<FMediaFrameworkCaptureCameraViewportCameraOutputInfo>(
+		UMediaFrameworkWorldSettingsAssetUserData* InSettings,
+		UMediaOutput* InMediaOutput)
+	{
+		TArray<FMediaFrameworkCaptureCameraViewportCameraOutputInfo> OutputInfos;
+		if (InSettings != nullptr)
+		{
+			for (const FMediaFrameworkCaptureCameraViewportCameraOutputInfo& OutputInfo : InSettings->ViewportCaptures)
+			{
+				if (OutputInfo.MediaOutput == InMediaOutput)
+				{
+					OutputInfos.Add(OutputInfo);
+				}
+			}
+		}
+	
+		return OutputInfos;
+	}
+
+	template <>
+	TArray<FMediaFrameworkCaptureRenderTargetCameraOutputInfo> FindAllOutputInfoForMediaOutput<FMediaFrameworkCaptureRenderTargetCameraOutputInfo>(
+		UMediaFrameworkWorldSettingsAssetUserData* InSettings,
+		UMediaOutput* InMediaOutput)
+	{
+		TArray<FMediaFrameworkCaptureRenderTargetCameraOutputInfo> OutputInfos;
+		if (InSettings != nullptr)
+		{
+			for (const FMediaFrameworkCaptureRenderTargetCameraOutputInfo& OutputInfo : InSettings->RenderTargetCaptures)
+			{
+				if (OutputInfo.MediaOutput == InMediaOutput)
+				{
+					OutputInfos.Add(OutputInfo);
+				}
+			}
+		}
+	
+		return OutputInfos;
+	}
+
+	template <>
+	TArray<FMediaFrameworkCaptureMediaTextureOutputInfo> FindAllOutputInfoForMediaOutput<FMediaFrameworkCaptureMediaTextureOutputInfo>(
+		UMediaFrameworkWorldSettingsAssetUserData* InSettings,
+		UMediaOutput* InMediaOutput)
+	{
+		TArray<FMediaFrameworkCaptureMediaTextureOutputInfo> OutputInfos;
+		if (InSettings != nullptr)
+		{
+			for (const FMediaFrameworkCaptureMediaTextureOutputInfo& OutputInfo : InSettings->MediaTextureCaptures)
+			{
+				if (OutputInfo.MediaOutput == InMediaOutput)
+				{
+					OutputInfos.Add(OutputInfo);
+				}
+			}
+		}
+	
+		return OutputInfos;
+	}
+
+	template <>
+	void ForEachOutputInfoForMediaOutput<FMediaFrameworkCaptureCurrentViewportOutputInfo>(
+		UMediaFrameworkWorldSettingsAssetUserData* InSettings, UMediaOutput* InMediaOutput,
+		TFunction<void(FMediaFrameworkCaptureCurrentViewportOutputInfo&)> ForEachFunc)
+	{
+		TArray<FMediaFrameworkCaptureCurrentViewportOutputInfo> OutputInfos;
+		if (InSettings != nullptr)
+		{
+			if (InSettings->CurrentViewportMediaOutput.MediaOutput == InMediaOutput)
+			{
+				ForEachFunc(InSettings->CurrentViewportMediaOutput);
+			}
+		}
+	}
+
+	template <>
+	void ForEachOutputInfoForMediaOutput<FMediaFrameworkCaptureCameraViewportCameraOutputInfo>(
+		UMediaFrameworkWorldSettingsAssetUserData* InSettings, UMediaOutput* InMediaOutput,
+		TFunction<void(FMediaFrameworkCaptureCameraViewportCameraOutputInfo&)> ForEachFunc)
+	{
+		TArray<FMediaFrameworkCaptureCameraViewportCameraOutputInfo> OutputInfos;
+		if (InSettings != nullptr)
+		{
+			for (FMediaFrameworkCaptureCameraViewportCameraOutputInfo& OutputInfo : InSettings->ViewportCaptures)
+			{
+				if (OutputInfo.MediaOutput == InMediaOutput)
+				{
+					ForEachFunc(OutputInfo);
+				}
+			}
+		}
+	}
+
+	template <>
+	void ForEachOutputInfoForMediaOutput<FMediaFrameworkCaptureRenderTargetCameraOutputInfo>(
+		UMediaFrameworkWorldSettingsAssetUserData* InSettings, UMediaOutput* InMediaOutput,
+		TFunction<void(FMediaFrameworkCaptureRenderTargetCameraOutputInfo&)> ForEachFunc)
+	{
+		TArray<FMediaFrameworkCaptureRenderTargetCameraOutputInfo> OutputInfos;
+		if (InSettings != nullptr)
+		{
+			for (FMediaFrameworkCaptureRenderTargetCameraOutputInfo& OutputInfo : InSettings->RenderTargetCaptures)
+			{
+				if (OutputInfo.MediaOutput == InMediaOutput)
+				{
+					ForEachFunc(OutputInfo);
+				}
+			}
+		}
+	}
+
+	template <>
+	void ForEachOutputInfoForMediaOutput<FMediaFrameworkCaptureMediaTextureOutputInfo>(
+		UMediaFrameworkWorldSettingsAssetUserData* InSettings, UMediaOutput* InMediaOutput,
+		TFunction<void(FMediaFrameworkCaptureMediaTextureOutputInfo&)> ForEachFunc)
+	{
+		TArray<FMediaFrameworkCaptureMediaTextureOutputInfo> OutputInfos;
+		if (InSettings != nullptr)
+		{
+			for (FMediaFrameworkCaptureMediaTextureOutputInfo& OutputInfo : InSettings->MediaTextureCaptures)
+			{
+				if (OutputInfo.MediaOutput == InMediaOutput)
+				{
+					ForEachFunc(OutputInfo);
+				}
+			}
+		}
+	}
+
+	template <>
+	int32 RemoveAllOutputInfoForMediaOutput<FMediaFrameworkCaptureCurrentViewportOutputInfo>(
+		UMediaFrameworkWorldSettingsAssetUserData* InSettings, UMediaOutput* InMediaOutput)
+	{
+		if (InSettings->CurrentViewportMediaOutput.MediaOutput == InMediaOutput)
+		{
+			InSettings->CurrentViewportMediaOutput.MediaOutput = nullptr;
+			return 1;
+		}
+		
+		return 0;
+	}
+
+	template <>
+	int32 RemoveAllOutputInfoForMediaOutput<FMediaFrameworkCaptureCameraViewportCameraOutputInfo>(
+		UMediaFrameworkWorldSettingsAssetUserData* InSettings, UMediaOutput* InMediaOutput)
+	{
+		return InSettings->ViewportCaptures.RemoveAll([InMediaOutput](const FMediaFrameworkCaptureCameraViewportCameraOutputInfo& OutputInfo)
+		{
+			return OutputInfo.MediaOutput == InMediaOutput;
+		});
+	}
+
+	template <>
+	int32 RemoveAllOutputInfoForMediaOutput<FMediaFrameworkCaptureRenderTargetCameraOutputInfo>(
+		UMediaFrameworkWorldSettingsAssetUserData* InSettings, UMediaOutput* InMediaOutput)
+	{
+		return InSettings->RenderTargetCaptures.RemoveAll([InMediaOutput](const FMediaFrameworkCaptureRenderTargetCameraOutputInfo& OutputInfo)
+		{
+			return OutputInfo.MediaOutput == InMediaOutput;
+		});
+	}
+
+	template <>
+	int32 RemoveAllOutputInfoForMediaOutput<FMediaFrameworkCaptureMediaTextureOutputInfo>(
+		UMediaFrameworkWorldSettingsAssetUserData* InSettings, UMediaOutput* InMediaOutput)
+	{
+		return InSettings->MediaTextureCaptures.RemoveAll([InMediaOutput](const FMediaFrameworkCaptureMediaTextureOutputInfo& OutputInfo)
+		{
+			return OutputInfo.MediaOutput == InMediaOutput;
+		});
+	}
+}
+
 PRAGMA_ENABLE_DEPRECATION_WARNINGS
diff --git a/Engine/Plugins/Media/MediaFrameworkUtilities/Source/MediaFrameworkUtilitiesEditor/Private/MediaFrameworkWorldSettingsAssetUserData.h b/Engine/Plugins/Media/MediaFrameworkUtilities/Source/MediaFrameworkUtilitiesEditor/Private/MediaFrameworkWorldSettingsAssetUserData.h
--- a/Engine/Plugins/Media/MediaFrameworkUtilities/Source/MediaFrameworkUtilitiesEditor/Private/MediaFrameworkWorldSettingsAssetUserData.h
+++ b/Engine/Plugins/Media/MediaFrameworkUtilities/Source/MediaFrameworkUtilitiesEditor/Private/MediaFrameworkWorldSettingsAssetUserData.h
@@ -114,6 +114,22 @@ struct FMediaFrameworkCaptureRenderTargetCameraOutputInfo
 	FMediaCaptureOptions CaptureOptions;
 };
 
+USTRUCT()
+struct FMediaFrameworkCaptureMediaTextureOutputInfo
+{
+	GENERATED_BODY()
+
+	FMediaFrameworkCaptureMediaTextureOutputInfo();
+
+	UPROPERTY(EditAnywhere, Category="Media Texture Capture")
+	TObjectPtr<UMediaTexture> MediaTexture;
+
+	UPROPERTY(EditAnywhere, Category="Media Texture Capture")
+	TObjectPtr<UMediaOutput> MediaOutput;
+
+	UPROPERTY(EditAnywhere, Category="Media Texture Capture")
+	FMediaCaptureOptions CaptureOptions;
+};
 
 /**
  * UMediaFrameworkCaptureCameraViewportAssetUserData
@@ -132,6 +148,9 @@ public:
 	UPROPERTY(EditAnywhere, config, Category="Media Viewport Capture", meta=(ShowOnlyInnerProperties))
 	TArray<FMediaFrameworkCaptureCameraViewportCameraOutputInfo> ViewportCaptures;
 
+	UPROPERTY(config)
+	TArray<FMediaFrameworkCaptureMediaTextureOutputInfo> MediaTextureCaptures;
+	
 	/**
 	 * Capture the current viewport. It may be the level editor active viewport or a PIE instance launch with "New Editor Window PIE".
 	 * @note The behavior is different from MediaCapture.CaptureActiveSceneViewport. Here we can capture the editor viewport (since we are in the editor).
@@ -143,3 +162,76 @@ public:
 public:
 	virtual void Serialize(FArchive& Ar) override;
 };
+
+namespace UE::MediaFrameworkWorldSettings::Helpers
+{
+	/**
+	 * Finds the first capture output info of a specific type in the capture settings that reference the specified media output
+	 * @tparam TOutputInfo The type of capture output info to search for
+	 * @param InSettings The capture settings to search in
+	 * @param InMediaOutput The media output to look for references for
+	 * @return A pointer to the found output info, or null if none were found
+	 */
+	template<typename TOutputInfo>
+	TOutputInfo* FindFirstOutputInfoForMediaOutput(UMediaFrameworkWorldSettingsAssetUserData* InSettings, UMediaOutput* InMediaOutput) { return nullptr; }
+	
+	// Template specializations
+	
+	template<> FMediaFrameworkCaptureCurrentViewportOutputInfo* FindFirstOutputInfoForMediaOutput(UMediaFrameworkWorldSettingsAssetUserData* InSettings, UMediaOutput* InMediaOutput);
+	template<> FMediaFrameworkCaptureCameraViewportCameraOutputInfo* FindFirstOutputInfoForMediaOutput(UMediaFrameworkWorldSettingsAssetUserData* InSettings, UMediaOutput* InMediaOutput);
+	template<> FMediaFrameworkCaptureRenderTargetCameraOutputInfo* FindFirstOutputInfoForMediaOutput(UMediaFrameworkWorldSettingsAssetUserData* InSettings, UMediaOutput* InMediaOutput);
+	template<> FMediaFrameworkCaptureMediaTextureOutputInfo* FindFirstOutputInfoForMediaOutput(UMediaFrameworkWorldSettingsAssetUserData* InSettings, UMediaOutput* InMediaOutput);
+	
+	/**
+	 * Finds all the capture output infos of a specific type in the capture settings that reference the specified media output
+	 * @tparam TOutputInfo The type of capture output info to search for
+	 * @param InSettings The capture settings to search in
+	 * @param InMediaOutput The media output to look for references for
+	 * @return A list of capture output infos of type TOutputInfo that reference InMediaOutput
+	 */
+	template<typename TOutputInfo>
+	TArray<TOutputInfo> FindAllOutputInfoForMediaOutput(UMediaFrameworkWorldSettingsAssetUserData* InSettings, UMediaOutput* InMediaOutput) { return TArray<TOutputInfo>(); }
+	
+	// Template specializations
+	
+	template<> TArray<FMediaFrameworkCaptureCurrentViewportOutputInfo> FindAllOutputInfoForMediaOutput(UMediaFrameworkWorldSettingsAssetUserData* InSettings, UMediaOutput* InMediaOutput);
+	template<> TArray<FMediaFrameworkCaptureCameraViewportCameraOutputInfo> FindAllOutputInfoForMediaOutput(UMediaFrameworkWorldSettingsAssetUserData* InSettings, UMediaOutput* InMediaOutput);
+	template<> TArray<FMediaFrameworkCaptureRenderTargetCameraOutputInfo> FindAllOutputInfoForMediaOutput(UMediaFrameworkWorldSettingsAssetUserData* InSettings, UMediaOutput* InMediaOutput);
+	template<> TArray<FMediaFrameworkCaptureMediaTextureOutputInfo> FindAllOutputInfoForMediaOutput(UMediaFrameworkWorldSettingsAssetUserData* InSettings, UMediaOutput* InMediaOutput);
+
+	/**
+	 * Calls the specified function on each capture output info of the specified type in the capture settings that references the specified media output
+	 * @tparam TOutputInfo The type of capture output info to search for
+	 * @param InSettings The capture settings to search in
+	 * @param InMediaOutput The media output to look for references for
+	 * @param ForEachFunc The function to run on each matching capture output info
+	 */
+	template<typename TOutputInfo>
+	void ForEachOutputInfoForMediaOutput(UMediaFrameworkWorldSettingsAssetUserData* InSettings, UMediaOutput* InMediaOutput, TFunction<void(TOutputInfo&)> ForEachFunc) { }
+	
+	// Template specializations
+	
+	template<> void ForEachOutputInfoForMediaOutput<FMediaFrameworkCaptureCurrentViewportOutputInfo>(
+		UMediaFrameworkWorldSettingsAssetUserData* InSettings, UMediaOutput* InMediaOutput, TFunction<void(FMediaFrameworkCaptureCurrentViewportOutputInfo&)> ForEachFunc);
+	template<> void ForEachOutputInfoForMediaOutput<FMediaFrameworkCaptureCameraViewportCameraOutputInfo>(
+		UMediaFrameworkWorldSettingsAssetUserData* InSettings, UMediaOutput* InMediaOutput, TFunction<void(FMediaFrameworkCaptureCameraViewportCameraOutputInfo&)> ForEachFunc);
+	template<> void ForEachOutputInfoForMediaOutput<FMediaFrameworkCaptureRenderTargetCameraOutputInfo>(
+		UMediaFrameworkWorldSettingsAssetUserData* InSettings, UMediaOutput* InMediaOutput, TFunction<void(FMediaFrameworkCaptureRenderTargetCameraOutputInfo&)> ForEachFunc);
+	template<> void ForEachOutputInfoForMediaOutput<FMediaFrameworkCaptureMediaTextureOutputInfo>(
+		UMediaFrameworkWorldSettingsAssetUserData* InSettings, UMediaOutput* InMediaOutput, TFunction<void(FMediaFrameworkCaptureMediaTextureOutputInfo&)> ForEachFunc);
+	
+	
+	template<typename TOutputInfo>
+	int32 RemoveAllOutputInfoForMediaOutput(UMediaFrameworkWorldSettingsAssetUserData* InSettings, UMediaOutput* InMediaOutput) { return 0; }
+	
+	// Template specializations
+	
+	template<> int32 RemoveAllOutputInfoForMediaOutput<FMediaFrameworkCaptureCurrentViewportOutputInfo>(
+		UMediaFrameworkWorldSettingsAssetUserData* InSettings, UMediaOutput* InMediaOutput);
+	template<> int32 RemoveAllOutputInfoForMediaOutput<FMediaFrameworkCaptureCameraViewportCameraOutputInfo>(
+		UMediaFrameworkWorldSettingsAssetUserData* InSettings, UMediaOutput* InMediaOutput);
+	template<> int32 RemoveAllOutputInfoForMediaOutput<FMediaFrameworkCaptureRenderTargetCameraOutputInfo>(
+		UMediaFrameworkWorldSettingsAssetUserData* InSettings, UMediaOutput* InMediaOutput);
+	template<> int32 RemoveAllOutputInfoForMediaOutput<FMediaFrameworkCaptureMediaTextureOutputInfo>(
+		UMediaFrameworkWorldSettingsAssetUserData* InSettings, UMediaOutput* InMediaOutput);
+}
\ No newline at end of file
diff --git a/Engine/Plugins/Media/MediaFrameworkUtilities/Source/MediaFrameworkUtilitiesEditor/Private/SMediaProfileSourceTexturePicker.cpp b/Engine/Plugins/Media/MediaFrameworkUtilities/Source/MediaFrameworkUtilitiesEditor/Private/SMediaProfileSourceTexturePicker.cpp
new file mode 100644
--- /dev/null
+++ b/Engine/Plugins/Media/MediaFrameworkUtilities/Source/MediaFrameworkUtilitiesEditor/Private/SMediaProfileSourceTexturePicker.cpp
@@ -0,0 +1,161 @@
+﻿// Copyright Epic Games, Inc. All Rights Reserved.
+
+#include "MediaFrameworkUtilitiesEditor/Public/SMediaProfileSourceTexturePicker.h"
+
+#include "Editor.h"
+#include "MediaSource.h"
+#include "MediaTexture.h"
+#include "Engine/Texture.h"
+#include "Engine/Texture2D.h"
+#include "Framework/MultiBox/MultiBoxBuilder.h"
+#include "Profile/IMediaProfileManager.h"
+#include "Profile/MediaProfile.h"
+#include "Profile/MediaProfilePlaybackManager.h"
+#include "Styling/SlateIconFinder.h"
+#include "Subsystems/AssetEditorSubsystem.h"
+#include "UI/MediaFrameworkUtilitiesEditorStyle.h"
+#include "Widgets/Images/SImage.h"
+#include "Widgets/Input/SComboButton.h"
+
+#define LOCTEXT_NAMESPACE "SMediaProfileSourceTexturePicker"
+
+void SMediaProfileSourceTexturePicker::Construct(const FArguments& InArgs)
+{
+	TexturePropertyHandle = InArgs._TexturePropertyHandle;
+	OnMediaSourceSelected = InArgs._OnMediaSourceSelected;
+	
+	SObjectPropertyEntryBox::FArguments BaseArgs;
+	BaseArgs
+		.PropertyHandle(TexturePropertyHandle)
+		.AllowedClass(UMediaTexture::StaticClass())
+		.AllowCreate(true)
+		.ThumbnailPool(InArgs._ThumbnailPool)
+		.OnShouldFilterAsset_Lambda([](FAssetData const& AssetData)
+		{
+			if (UClass* AssetClass = AssetData.GetClass())
+			{
+				return !AssetClass->IsChildOf(UMediaTexture::StaticClass());
+			}
+
+			return true;
+		})
+		.CustomContentSlot()
+		[
+			SNew(SHorizontalBox)
+
+			+SHorizontalBox::Slot()
+			.AutoWidth()
+			[
+				SNew(SComboButton)
+				.HasDownArrow(true)
+				.ComboButtonStyle(FAppStyle::Get(), "SimpleComboButton")
+				.ToolTipText(LOCTEXT("MediaProfileSourceSelectorToolTip", "Select media texture from a media source in the active Media Profile"))
+				.HAlign(HAlign_Center)
+				.VAlign(VAlign_Center)
+				.OnGetMenuContent(this, &SMediaProfileSourceTexturePicker::GetMediaProfileSourceSelectorMenu)
+				.IsEnabled(this, &SMediaProfileSourceTexturePicker::HasActiveMediaProfile)
+				.ButtonContent()
+				[
+					SNew(SImage)
+					.DesiredSizeOverride(FVector2D(16.0f, 16.0f))
+					.Image(FMediaFrameworkUtilitiesEditorStyle::Get().GetBrush("ClassIcon.MediaProfile"))
+				]
+			]
+		];
+	
+	SObjectPropertyEntryBox::Construct(BaseArgs);
+}
+
+TSharedRef<SWidget> SMediaProfileSourceTexturePicker::GetMediaProfileSourceSelectorMenu()
+{
+	UMediaProfile* ActiveMediaProfile = IMediaProfileManager::Get().GetCurrentMediaProfile();
+	if (!ActiveMediaProfile)
+	{
+		return SNullWidget::NullWidget;
+	}
+	
+	const bool bCloseAfterSelection = true;
+	FMenuBuilder MenuBuilder(bCloseAfterSelection, nullptr);
+
+	MenuBuilder.BeginSection(NAME_None, LOCTEXT("MediaProfileSourcesSectionLabel", "Texture from Media Profile"));
+	{
+		bool bNoMediaSources = true;
+		for (int32 Index = 0; Index < ActiveMediaProfile->NumMediaSources(); ++Index)
+		{
+			if (UMediaSource* MediaSource = ActiveMediaProfile->GetMediaSource(Index))
+			{
+				bNoMediaSources = false;
+
+				MenuBuilder.AddMenuEntry(
+					FText::FromString(ActiveMediaProfile->GetLabelForMediaSource(Index)),
+					FText::GetEmpty(),
+					FSlateIconFinder::FindIconForClass(MediaSource->GetClass()),
+					FUIAction(FExecuteAction::CreateSP(this, &SMediaProfileSourceTexturePicker::SetTextureToMediaProfileSource, Index))
+				);
+			}
+		}
+
+		if (bNoMediaSources)
+		{
+			TSharedRef<SWidget> TextBlock = SNew(STextBlock)
+				.Text(LOCTEXT("NoMediaSources", "No Media Sources"))
+				.Font(FCoreStyle::GetDefaultFontStyle("Italic", 10));
+			
+			MenuBuilder.AddWidget(TextBlock,FText::GetEmpty());
+		}
+	}
+	MenuBuilder.EndSection();
+
+	MenuBuilder.AddSeparator();
+	MenuBuilder.BeginSection(NAME_None);
+	{
+		MenuBuilder.AddMenuEntry(
+			FText::Format(LOCTEXT("OpenActiveMediaProfileLabelFormat", "Open {0}"), FText::FromString(ActiveMediaProfile->GetName())),
+			FText::GetEmpty(),
+			FSlateIconFinder::FindIconForClass(ActiveMediaProfile->StaticClass()),
+			FUIAction(FExecuteAction::CreateSP(this, &SMediaProfileSourceTexturePicker::OpenMediaProfile))
+		);
+	}
+	MenuBuilder.EndSection();
+	
+	return MenuBuilder.MakeWidget();
+}
+
+bool SMediaProfileSourceTexturePicker::HasActiveMediaProfile() const
+{
+	UMediaProfile* ActiveMediaProfile = IMediaProfileManager::Get().GetCurrentMediaProfile();
+	return IsValid(ActiveMediaProfile);
+}
+
+void SMediaProfileSourceTexturePicker::SetTextureToMediaProfileSource(int32 InMediaSourceIndex)
+{
+	UMediaProfile* ActiveMediaProfile = IMediaProfileManager::Get().GetCurrentMediaProfile();
+	if (!ActiveMediaProfile)
+	{
+		return;
+	}
+
+	UMediaSource* MediaSource = ActiveMediaProfile->GetMediaSource(InMediaSourceIndex);
+	if (MediaSource == nullptr)
+	{
+		return;
+	}
+
+	UMediaTexture* MediaTexture = ActiveMediaProfile->GetPlaybackManager()->GetSourceMediaTextureFromIndex(InMediaSourceIndex);
+	TexturePropertyHandle->SetValue(MediaTexture);
+
+	OnMediaSourceSelected.ExecuteIfBound(ActiveMediaProfile, InMediaSourceIndex);
+}
+
+void SMediaProfileSourceTexturePicker::OpenMediaProfile()
+{
+	UMediaProfile* ActiveMediaProfile = IMediaProfileManager::Get().GetCurrentMediaProfile();
+	if (!ActiveMediaProfile)
+	{
+		return;
+	}
+	
+	GEditor->GetEditorSubsystem<UAssetEditorSubsystem>()->OpenEditorForAsset(ActiveMediaProfile);
+}
+
+#undef LOCTEXT_NAMESPACE
diff --git a/Engine/Plugins/Media/MediaFrameworkUtilities/Source/MediaFrameworkUtilitiesEditor/Private/UI/MediaFrameworkUtilitiesEditorStyle.cpp b/Engine/Plugins/Media/MediaFrameworkUtilities/Source/MediaFrameworkUtilitiesEditor/Private/UI/MediaFrameworkUtilitiesEditorStyle.cpp
--- a/Engine/Plugins/Media/MediaFrameworkUtilities/Source/MediaFrameworkUtilitiesEditor/Private/UI/MediaFrameworkUtilitiesEditorStyle.cpp
+++ b/Engine/Plugins/Media/MediaFrameworkUtilities/Source/MediaFrameworkUtilitiesEditor/Private/UI/MediaFrameworkUtilitiesEditorStyle.cpp
@@ -38,7 +38,7 @@ void FMediaFrameworkUtilitiesEditorStyle::Register()
 	MediaFrameworkUtilitiesStyle::StyleInstance->Set("ClassThumbnail.ProxyMediaSource", new IMAGE_BRUSH("ProxyMediaSource_64x", MediaFrameworkUtilitiesStyle::Icon64x64));
 	MediaFrameworkUtilitiesStyle::StyleInstance->Set("ClassIcon.ProxyMediaSource", new IMAGE_BRUSH("ProxyMediaSource_16x", MediaFrameworkUtilitiesStyle::Icon16x16));
 	MediaFrameworkUtilitiesStyle::StyleInstance->Set("ClassThumbnail.MediaProfile", new IMAGE_BRUSH("MediaProfile_64x", MediaFrameworkUtilitiesStyle::Icon64x64));
-	MediaFrameworkUtilitiesStyle::StyleInstance->Set("ClassIcon.MediaProfile", new IMAGE_BRUSH("MediaProfile_20x", MediaFrameworkUtilitiesStyle::Icon20x20));
+	MediaFrameworkUtilitiesStyle::StyleInstance->Set("ClassIcon.MediaProfile", new IMAGE_BRUSH_SVG("MediaProfile", MediaFrameworkUtilitiesStyle::Icon20x20));
 
 	MediaFrameworkUtilitiesStyle::StyleInstance->Set("ToolbarIcon.Edit", new IMAGE_BRUSH_SVG("MediaProfile", MediaFrameworkUtilitiesStyle::Icon40x40));
 	MediaFrameworkUtilitiesStyle::StyleInstance->Set("ToolbarIcon.Apply", new IMAGE_BRUSH("Apply_40x", MediaFrameworkUtilitiesStyle::Icon40x40));
diff --git a/Engine/Plugins/Media/MediaFrameworkUtilities/Source/MediaFrameworkUtilitiesEditor/Public/SMediaProfileSourceTexturePicker.h b/Engine/Plugins/Media/MediaFrameworkUtilities/Source/MediaFrameworkUtilitiesEditor/Public/SMediaProfileSourceTexturePicker.h
new file mode 100644
--- /dev/null
+++ b/Engine/Plugins/Media/MediaFrameworkUtilities/Source/MediaFrameworkUtilitiesEditor/Public/SMediaProfileSourceTexturePicker.h
@@ -0,0 +1,46 @@
+﻿// Copyright Epic Games, Inc. All Rights Reserved.
+
+#pragma once
+
+#include "PropertyCustomizationHelpers.h"
+
+class UMediaProfile;
+
+/**
+ * A SObjectPropertyEntryBox for picking a texture that adds a dropdown for picking the media texture from a media source in the active media profile
+ */
+class MEDIAFRAMEWORKUTILITIESEDITOR_API SMediaProfileSourceTexturePicker : public SObjectPropertyEntryBox
+{
+public:
+	DECLARE_DELEGATE_TwoParams(FOnMediaSourceSelected, UMediaProfile* /* InMediaProfile */, int32 /* InMediaSourceIndex */);
+	
+public:
+	SLATE_BEGIN_ARGS(SMediaProfileSourceTexturePicker) { }
+		SLATE_ARGUMENT(TSharedPtr<IPropertyHandle>, TexturePropertyHandle)
+		SLATE_ARGUMENT(TSharedPtr<FAssetThumbnailPool>, ThumbnailPool)
+		SLATE_EVENT(FOnMediaSourceSelected, OnMediaSourceSelected)
+	SLATE_END_ARGS()
+
+	/** Constructs this widget with InArgs */
+	void Construct(const FArguments& InArgs);
+
+private:
+	/** Constructs the dropdown menu that shows all the media sources in the active media profile */
+	TSharedRef<SWidget> GetMediaProfileSourceSelectorMenu();
+	
+	/** Gets whether there is a valid active media profile */
+	bool HasActiveMediaProfile() const;
+	
+	/** Raised when the user has selected a media source whose media texture should be used from the media sources dropdown */
+	void SetTextureToMediaProfileSource(int32 InMediaSourceIndex);
+	
+	/** Opens the active media profile in the media profile editor */
+	void OpenMediaProfile();
+
+private:
+	/** The property handle to store the texture in */
+	TSharedPtr<IPropertyHandle> TexturePropertyHandle;
+
+	/** Raised when a media source is selected */
+	FOnMediaSourceSelected OnMediaSourceSelected;
+};
diff --git a/Engine/Plugins/Media/MediaIOFramework/Source/MediaIOCore/Private/MediaCapture.cpp b/Engine/Plugins/Media/MediaIOFramework/Source/MediaIOCore/Private/MediaCapture.cpp
--- a/Engine/Plugins/Media/MediaIOFramework/Source/MediaIOCore/Private/MediaCapture.cpp
+++ b/Engine/Plugins/Media/MediaIOFramework/Source/MediaIOCore/Private/MediaCapture.cpp
@@ -264,6 +264,12 @@ bool UMediaCapture::CaptureTextureRenderTarget2D(UTextureRenderTarget2D* RenderT
 	return StartSourceCapture(MakeShared<FRenderTargetCaptureSource>(this, CaptureOptions, RenderTarget));
 }
 
+bool UMediaCapture::CaptureMediaTexture(UMediaTexture* InMediaTexture, FMediaCaptureOptions CaptureOptions)
+{
+	using namespace UE::MediaCapture::Private;
+	return StartSourceCapture(MakeShared<FMediaTextureCaptureSource>(this, CaptureOptions, InMediaTexture));
+}
+
 bool UMediaCapture::StartSourceCapture(TSharedPtr<UE::MediaCapture::Private::FCaptureSource> InSource)
 {
 	StopCapture(false);
diff --git a/Engine/Plugins/Media/MediaIOFramework/Source/MediaIOCore/Private/MediaCaptureHelper.cpp b/Engine/Plugins/Media/MediaIOFramework/Source/MediaIOCore/Private/MediaCaptureHelper.cpp
--- a/Engine/Plugins/Media/MediaIOFramework/Source/MediaIOCore/Private/MediaCaptureHelper.cpp
+++ b/Engine/Plugins/Media/MediaIOFramework/Source/MediaIOCore/Private/MediaCaptureHelper.cpp
@@ -27,16 +27,19 @@ bool FMediaCaptureHelper::AreInputsValid(const UE::MediaCaptureData::FCaptureFra
 {
 	// If it is a simple rgba swizzle we can handle the conversion. Supported formats
 	// contained in SupportedRgbaSwizzleFormats. Warning would've been displayed on start of capture.
-	if (Args.MediaCapture->GetDesiredPixelFormat() != Args.GetFormat() &&
-		(!UMediaCapture::GetSupportedRgbaSwizzleFormats().Contains(Args.GetFormat()) || !Args.MediaCapture->DesiredCaptureOptions.bConvertToDesiredPixelFormat))
+	if (Args.MediaCapture->GetDesiredPixelFormat() != Args.GetFormat())
 	{
-		UE_LOG(LogMediaIOCore, Error, TEXT("The capture will stop for '%s'. The Source pixel format doesn't match with the user requested pixel format. %sRequested: %s Source: %s")
-			, *Args.MediaCapture->MediaOutputName
-			, (UMediaCapture::GetSupportedRgbaSwizzleFormats().Contains(Args.GetFormat()) && !Args.MediaCapture->DesiredCaptureOptions.bConvertToDesiredPixelFormat) ? TEXT("Please enable \"Convert To Desired Pixel Format\" option in Media Capture settings. ") : TEXT("")
-			, GetPixelFormatString(Args.MediaCapture->GetDesiredPixelFormat())
-			, GetPixelFormatString(Args.GetFormat()));
+		const bool bIsCompressedFormat = IsDXTCBlockCompressedTextureFormat(Args.GetFormat());
+		if (!bIsCompressedFormat && (!UMediaCapture::GetSupportedRgbaSwizzleFormats().Contains(Args.GetFormat()) || !Args.MediaCapture->DesiredCaptureOptions.bConvertToDesiredPixelFormat))
+		{
+			UE_LOG(LogMediaIOCore, Error, TEXT("The capture will stop for '%s'. The Source pixel format doesn't match with the user requested pixel format. %sRequested: %s Source: %s")
+				, *Args.MediaCapture->MediaOutputName
+				, (UMediaCapture::GetSupportedRgbaSwizzleFormats().Contains(Args.GetFormat()) && !Args.MediaCapture->DesiredCaptureOptions.bConvertToDesiredPixelFormat) ? TEXT("Please enable \"Convert To Desired Pixel Format\" option in Media Capture settings. ") : TEXT("")
+				, GetPixelFormatString(Args.MediaCapture->GetDesiredPixelFormat())
+				, GetPixelFormatString(Args.GetFormat()));
 
-		return false;
+			return false;
+		}
 	}
 
 	if (Args.MediaCapture->DesiredCaptureOptions.ResizeMethod == EMediaCaptureResizeMethod::ResizeInCapturePass &&
diff --git a/Engine/Plugins/Media/MediaIOFramework/Source/MediaIOCore/Private/MediaCaptureRenderPass.cpp b/Engine/Plugins/Media/MediaIOFramework/Source/MediaIOCore/Private/MediaCaptureRenderPass.cpp
--- a/Engine/Plugins/Media/MediaIOFramework/Source/MediaIOCore/Private/MediaCaptureRenderPass.cpp
+++ b/Engine/Plugins/Media/MediaIOFramework/Source/MediaIOCore/Private/MediaCaptureRenderPass.cpp
@@ -20,6 +20,7 @@
 DECLARE_GPU_STAT(MediaCapture_Resample);
 DECLARE_GPU_STAT(MediaCapture_ColorConversion);
 DECLARE_GPU_STAT(MediaCapture_Conversion);
+DECLARE_GPU_STAT(MediaCapture_Uncompression);
 
 namespace UE::MediaCapture::Resample
 {
@@ -375,11 +376,64 @@ namespace UE::MediaCapture
 		return ConversionPass;
 	}
 
+	FRenderPass FRenderPass::CreateUncompressionPass(const UMediaCapture* MediaCapture)
+	{
+		FRenderPass UncompressionPass;
+		UncompressionPass.Name = "Uncompression";
+		UncompressionPass.OutputType = ERDGViewableResourceType::Texture;
+		UncompressionPass.ShouldExecutePassDelegate = FRenderPass::FShouldExecutePass::CreateLambda([](const UE::MediaCaptureData::FCaptureFrameArgs& Args, FRDGViewableResource* InputResource)
+		{
+			check(InputResource->Type == ERDGViewableResourceType::Texture);
+			FRDGTextureRef SourceRGBTexture = static_cast<FRDGTextureRef>(InputResource);
+			return IsBlockCompressedFormat(SourceRGBTexture->Desc.Format);
+		});
+		UncompressionPass.InitializePassOutputDelegate = FRenderPass::FInitializePassOutput::CreateLambda([](const UE::MediaCapture::FInitializePassOutputArgs& Args, uint32 FrameId)
+		{
+			check(Args.InputResource && Args.InputResource->Type == ERDGViewableResourceType::Texture);
+			FRDGTextureDesc InputResourceDesc = static_cast<FRDGTextureRef>(Args.InputResource)->Desc;
+			
+			FRDGTextureDesc UncompressedOutputTextureDesc = InputResourceDesc;
+			UncompressedOutputTextureDesc.Reset();
+			UncompressedOutputTextureDesc.NumMips = 1;
+			UncompressedOutputTextureDesc.Flags = TexCreate_None;
+			UncompressedOutputTextureDesc.Flags |= TexCreate_RenderTargetable | TexCreate_ShaderResource | TexCreate_UAV | TexCreate_NoFastClear;
+			UncompressedOutputTextureDesc.ClearValue = FClearValueBinding(FLinearColor(0, 0, 0, 0));
+			UncompressedOutputTextureDesc.Format = PF_FloatRGBA;
+
+			TRefCountPtr<IPooledRenderTarget> RenderTarget = AllocatePooledTexture(UncompressedOutputTextureDesc, TEXT("MediaCapture Uncompressed Output"));
+			return RenderTargetResource(MoveTemp(RenderTarget));
+		});
+
+		UncompressionPass.ExecutePassDelegate = FRenderPass::FExecutePass::CreateLambda([](const UE::MediaCaptureData::FCaptureFrameArgs& Args, const TSharedPtr<UE::MediaCaptureData::FCaptureFrame>& CapturingFrame, FRDGViewableResource* InputResource, FRDGViewableResource* OutputResource)
+		{
+			// At the moment we only support Texture inputs
+			check(InputResource->Type == ERDGViewableResourceType::Texture);
+
+			FRHICopyTextureInfo CopyInfo;
+			FVector2D SizeU;
+			FVector2D SizeV;
+			GetCopyInfo(Args, CopyInfo, SizeU, SizeV);
+
+			RDG_EVENT_SCOPE_STAT(Args.GraphBuilder, MediaCapture_Uncompression, "MediaCapture_Uncompression");
+			TRACE_CPUPROFILER_EVENT_SCOPE(UMediaCapture::Uncompression);
+			
+			FBaseRenderPassArgs UncompressionPassArgs;
+			UncompressionPassArgs.SourceRGBTexture = static_cast<FRDGTextureRef>(InputResource);
+			UncompressionPassArgs.CopyInfo = CopyInfo;
+			
+			AddUncompressionPass(Args, UncompressionPassArgs, OutputResource);
+		});
+		
+		return UncompressionPass;
+	}
+
 	/** Create the necessary render passes needed for the current media capture configuration. */
 	TArray<FRenderPass> FRenderPass::CreateRenderPasses(const UMediaCapture* MediaCapture)
 	{
 		TArray<FRenderPass> RenderPasses;
 
+		RenderPasses.Add(CreateUncompressionPass(MediaCapture));
+
 		if (MediaCapture->GetDesiredCaptureOptions().ResizeMethod == EMediaCaptureResizeMethod::ResizeInRenderPass)
 		{
 			if (MediaCapture->GetDesiredCaptureOptions().Crop != EMediaCaptureCroppingType::None)
@@ -484,6 +538,20 @@ namespace UE::MediaCapture
 		}
 	}
 	
+	void FRenderPass::AddUncompressionPass(const UE::MediaCaptureData::FCaptureFrameArgs& Args, const FBaseRenderPassArgs& UncompressionPassArgs, FRDGViewableResource* OutputResource)
+	{
+		if (OutputResource->Type == FRDGBuffer::StaticType)
+		{
+			return;
+		}
+
+		check(OutputResource->Type == FRDGTexture::StaticType);
+		FRDGTexture* OutputTexture =  static_cast<FRDGTexture*>(OutputResource);
+
+		// Very basic uncompression that draws the compressed texture into a PF_FloatRGBA texture
+		AddDrawTexturePass(Args.GraphBuilder, GetGlobalShaderMap(GMaxRHIFeatureLevel), UncompressionPassArgs.SourceRGBTexture, OutputTexture, FRDGDrawTextureInfo());
+	}
+
 	void FRenderPass::AddTextureResamplePass(const  UE::MediaCaptureData::FCaptureFrameArgs& Args, const FBaseRenderPassArgs& ConversionPassArgs, FRDGViewableResource* OutputTexture)
 	{
 		// Downsampling only supports texture at the moment.
diff --git a/Engine/Plugins/Media/MediaIOFramework/Source/MediaIOCore/Private/MediaCaptureRenderPass.h b/Engine/Plugins/Media/MediaIOFramework/Source/MediaIOCore/Private/MediaCaptureRenderPass.h
--- a/Engine/Plugins/Media/MediaIOFramework/Source/MediaIOCore/Private/MediaCaptureRenderPass.h
+++ b/Engine/Plugins/Media/MediaIOFramework/Source/MediaIOCore/Private/MediaCaptureRenderPass.h
@@ -68,6 +68,8 @@ namespace UE::MediaCapture
 		static void AddColorConversionPass(const UE::MediaCaptureData::FCaptureFrameArgs& Args, const TSharedPtr<UE::MediaCaptureData::FCaptureFrame> CapturingFrame, const FColorConversionPassArgs& ConversionPassArgs, FRDGViewableResource* OutputResource);
 		/** Add a format conversion pass to the graph builder. */
 		static void AddConversionPass(const UE::MediaCaptureData::FCaptureFrameArgs& Args, const FConversionPassArgs& ConversionPassArgs, FRDGViewableResource* OutputResource);
+		/** Add a pass to uncompress a compressed texture into a proper uncompressed format (currently PF_FloatRGBA) */
+		static void AddUncompressionPass(const UE::MediaCaptureData::FCaptureFrameArgs& Args, const FBaseRenderPassArgs& UncompressionPassArgs, FRDGViewableResource* OutputResource);
 
 		/** This delegate should return a TRDGResource constructed with either a pooled texture or buffer. */
 		DECLARE_DELEGATE_RetVal_TwoParams(TRDGResource /*Output*/, FInitializePassOutput, const UE::MediaCapture::FInitializePassOutputArgs& /*InitializePassOutputArgs*/, uint32 /*FrameId*/);
@@ -103,6 +105,8 @@ namespace UE::MediaCapture
 		static FRenderPass CreateResamplePass();
 		/** Creates a format conversion pass according to the capture settings. */
 		static FRenderPass CreateFormatConversionPass(const UMediaCapture* MediaCapture);
+		/** Creates an uncompression pass according to the capture settings. */
+		static FRenderPass CreateUncompressionPass(const UMediaCapture* MediaCapture);
 	};
 
 	/** Instantiates a TRDGResource TVariant from a pooled render target. */ 
diff --git a/Engine/Plugins/Media/MediaIOFramework/Source/MediaIOCore/Private/MediaCaptureSources.inl b/Engine/Plugins/Media/MediaIOFramework/Source/MediaIOCore/Private/MediaCaptureSources.inl
--- a/Engine/Plugins/Media/MediaIOFramework/Source/MediaIOCore/Private/MediaCaptureSources.inl
+++ b/Engine/Plugins/Media/MediaIOFramework/Source/MediaIOCore/Private/MediaCaptureSources.inl
@@ -9,6 +9,7 @@
 #include "MediaCapture.h"
 #include "MediaIOCoreModule.h"
 #include "MediaOutput.h"
+#include "MediaTexture.h"
 #include "Slate/SceneViewport.h"
 #include "TextureResource.h"
 #include "UObject/WeakObjectPtr.h"
@@ -132,6 +133,23 @@ namespace UE::MediaCapture::Private
 		return true;
 	}
 
+	bool ValidateMediaTexture(const UMediaTexture* InMediaTexture, const FMediaCaptureOptions& CaptureOptions, const FIntPoint& DesiredSize, const EPixelFormat DesiredPixelFormat, const bool bCurrentlyCapturing)
+	{
+		if (InMediaTexture == nullptr)
+		{
+			UE_LOG(LogMediaIOCore, Error, TEXT("Couldn't %s the capture. The Media Texture is invalid.")
+				, bCurrentlyCapturing ? TEXT("continue") : TEXT("start"));
+			return false;
+		}
+
+		if (CaptureOptions.ResizeMethod != EMediaCaptureResizeMethod::ResizeInRenderPass && !ValidateSize(FIntPoint(InMediaTexture->GetWidth(), InMediaTexture->GetHeight()), DesiredSize, CaptureOptions, bCurrentlyCapturing))
+		{
+			return false;
+		}
+
+		return true;
+	}
+
 	enum class ECaptureType : uint8
 	{
 		Immediate,
@@ -422,5 +440,73 @@ namespace UE::MediaCapture::Private
 
 		TWeakObjectPtr<UTextureRenderTarget2D> RenderTarget;
 	};
+
+	class FMediaTextureCaptureSource : public FCaptureSource
+	{
+	public:
+		FMediaTextureCaptureSource(UMediaCapture* InMediaCapture, const FMediaCaptureOptions& InCaptureOptions, UMediaTexture* InMediaTexture)
+			: FCaptureSource(InMediaCapture, InCaptureOptions)
+			, MediaTexture(InMediaTexture)
+		{
+		}
+
+		virtual bool ValidateSource(const FIntPoint& DesiredSize, const EPixelFormat DesiredPixelFormat, const bool bCurrentlyCapturing) override
+		{
+			return ValidateMediaTexture(MediaTexture.Get(), CaptureOptions, DesiredSize, DesiredPixelFormat, bCurrentlyCapturing);
+		}
+		
+		virtual bool PostInitialize() override
+		{
+			if (TStrongObjectPtr<UMediaCapture> PinnedMediaCapture = MediaCapture.Pin())
+			{
+				return PinnedMediaCapture->PostInitializeCaptureMediaTexture(MediaTexture.Get());
+			}
+			return false;
+		}
+		
+		virtual EMediaCaptureSourceType GetSourceType() override
+		{
+			return EMediaCaptureSourceType::MEDIA_TEXTURE;
+		}
+		
+		virtual FIntPoint GetSize() override
+		{
+			if (TStrongObjectPtr<UMediaTexture> PinnedMediaTexture = MediaTexture.Pin())
+			{
+				return FIntPoint(PinnedMediaTexture->GetWidth(), PinnedMediaTexture->GetHeight());
+			}
+
+			return FIntPoint();
+		}
+		
+		virtual bool IsValid() const override
+		{
+			return MediaTexture.IsValid();
+		}
+
+		virtual bool UpdateSourceImpl() override
+		{
+			if (TStrongObjectPtr<UMediaCapture> PinnedMediaCapture = MediaCapture.Pin())
+			{
+				return PinnedMediaCapture->UpdateMediaTextureImpl(MediaTexture.Get());
+			}
+			return false;
+		}
+
+		virtual FTextureRHIRef GetSourceTextureForInput_RenderThread() override
+		{
+			constexpr bool bEvenIfPendingKill = false;
+			constexpr bool bThreadSafeTest = true;
+			if (MediaTexture.IsValid(bEvenIfPendingKill, bThreadSafeTest) && MediaTexture.GetEvenIfUnreachable()->GetResource())
+			{
+				return MediaTexture.GetEvenIfUnreachable()->GetResource()->GetTextureRHI();
+			}
+			
+			return nullptr;
+		}
+		
+	private:
+		TWeakObjectPtr<UMediaTexture> MediaTexture;
+	};
 }
 
diff --git a/Engine/Plugins/Media/MediaIOFramework/Source/MediaIOCore/Private/RenderTargetMediaCapture.cpp b/Engine/Plugins/Media/MediaIOFramework/Source/MediaIOCore/Private/RenderTargetMediaCapture.cpp
--- a/Engine/Plugins/Media/MediaIOFramework/Source/MediaIOCore/Private/RenderTargetMediaCapture.cpp
+++ b/Engine/Plugins/Media/MediaIOFramework/Source/MediaIOCore/Private/RenderTargetMediaCapture.cpp
@@ -5,6 +5,7 @@
 #include "Engine/Engine.h"
 #include "Engine/TextureRenderTarget2D.h"
 #include "MediaIOCoreModule.h"
+#include "MediaTexture.h"
 #include "RenderTargetMediaOutput.h"
 #include "Slate/SceneViewport.h"
 #include "TextureResource.h"
@@ -74,6 +75,11 @@ bool URenderTargetMediaCapture::PostInitializeCaptureRHIResource(const FRHICaptu
 	return StartNewCapture(InResourceDescription.ResourceSize, InResourceDescription.PixelFormat);
 }
 
+bool URenderTargetMediaCapture::PostInitializeCaptureMediaTexture(UMediaTexture* InMediaTexture)
+{
+	return StartNewCapture(FIntPoint(InMediaTexture->GetWidth(), InMediaTexture->GetHeight()), EPixelFormat::PF_B8G8R8A8);
+}
+
 void URenderTargetMediaCapture::StopCaptureImpl(bool /*bAllowPendingFrameToBeProcess*/)
 {
 	TRACE_BOOKMARK(TEXT("RenderTargetMediaCapture::StopCapture"));
diff --git a/Engine/Plugins/Media/MediaIOFramework/Source/MediaIOCore/Private/RenderTargetMediaCapture.h b/Engine/Plugins/Media/MediaIOFramework/Source/MediaIOCore/Private/RenderTargetMediaCapture.h
--- a/Engine/Plugins/Media/MediaIOFramework/Source/MediaIOCore/Private/RenderTargetMediaCapture.h
+++ b/Engine/Plugins/Media/MediaIOFramework/Source/MediaIOCore/Private/RenderTargetMediaCapture.h
@@ -32,6 +32,7 @@ protected:
 	virtual bool PostInitializeCaptureViewport(TSharedPtr<FSceneViewport>& InSceneViewport) override;
 	virtual bool PostInitializeCaptureRenderTarget(UTextureRenderTarget2D* InRenderTarget) override;
 	virtual bool PostInitializeCaptureRHIResource(const FRHICaptureResourceDescription& InResourceDescription) override;
+	virtual bool PostInitializeCaptureMediaTexture(UMediaTexture* InMediaTexture) override;
 	virtual void StopCaptureImpl(bool bAllowPendingFrameToBeProcess) override;
 
 	UPROPERTY(Transient)
diff --git a/Engine/Plugins/Media/MediaIOFramework/Source/MediaIOCore/Public/MediaCapture.h b/Engine/Plugins/Media/MediaIOFramework/Source/MediaIOCore/Public/MediaCapture.h
--- a/Engine/Plugins/Media/MediaIOFramework/Source/MediaIOCore/Public/MediaCapture.h
+++ b/Engine/Plugins/Media/MediaIOFramework/Source/MediaIOCore/Public/MediaCapture.h
@@ -315,6 +315,15 @@ public:
 	UFUNCTION(BlueprintCallable, Category = "Media|Output")
 	UE_API bool CaptureTextureRenderTarget2D(UTextureRenderTarget2D* RenderTarget, FMediaCaptureOptions CaptureOptions);
 
+	/**
+	 * Stop the actual capture if there is one.
+	 * Then capture every frame for a MediaTexture.
+	 * The MediaTexture needs to be of the same size and have the same pixel format as requested by the media output.
+	 * @return True if the capture was successfully started
+	 */
+	UFUNCTION(BlueprintCallable, Category = "Media|Output")
+	UE_API bool CaptureMediaTexture(UMediaTexture* InMediaTexture, FMediaCaptureOptions CaptureOptions);
+	
 	/**
 	 * Update the current capture with a SceneViewport.
 	 * If the SceneViewport is destroyed, the capture will stop.
@@ -425,9 +434,13 @@ public:
 
 	/** Called after initialize for rhi resource capture type */
 	virtual bool PostInitializeCaptureRHIResource(const FRHICaptureResourceDescription& InResourceDescription) { return true; }
+
+	/** Called after initialize for media texture capture type */
+	virtual bool PostInitializeCaptureMediaTexture(UMediaTexture* InMediaTexture) { return true; }
 	
 	virtual bool UpdateSceneViewportImpl(TSharedPtr<FSceneViewport>& InSceneViewport) { return true; }
 	virtual bool UpdateRenderTargetImpl(UTextureRenderTarget2D* InRenderTarget) { return true; }
+	virtual bool UpdateMediaTextureImpl(UMediaTexture* InMediaTexture) { return true; }
 	virtual bool UpdateAudioDeviceImpl(const FAudioDeviceHandle& InAudioDeviceHandle) { return true; }
 
 	/**
diff --git a/Engine/Plugins/Media/MediaIOFramework/Source/MediaIOCore/Public/MediaOutput.h b/Engine/Plugins/Media/MediaIOFramework/Source/MediaIOCore/Public/MediaOutput.h
--- a/Engine/Plugins/Media/MediaIOFramework/Source/MediaIOCore/Public/MediaOutput.h
+++ b/Engine/Plugins/Media/MediaIOFramework/Source/MediaIOCore/Public/MediaOutput.h
@@ -31,6 +31,7 @@ enum class EMediaCaptureSourceType : int8
 	RENDER_TARGET,
 	SCENE_VIEWPORT,
 	RHI_RESOURCE,
+	MEDIA_TEXTURE
 };
 
 /**
diff --git a/Engine/Plugins/Media/NDIMedia/Source/NDIMedia/Private/NDIMediaCapture.cpp b/Engine/Plugins/Media/NDIMedia/Source/NDIMedia/Private/NDIMediaCapture.cpp
--- a/Engine/Plugins/Media/NDIMedia/Source/NDIMedia/Private/NDIMediaCapture.cpp
+++ b/Engine/Plugins/Media/NDIMedia/Source/NDIMedia/Private/NDIMediaCapture.cpp
@@ -5,6 +5,7 @@
 #include "Engine/Engine.h"
 #include "Engine/TextureRenderTarget2D.h"
 #include "MediaIOCoreSubsystem.h"
+#include "MediaTexture.h"
 #include "NDIMediaAPI.h"
 #include "NDIMediaLog.h"
 #include "NDIMediaModule.h"
@@ -427,6 +428,17 @@ bool UNDIMediaCapture::PostInitializeCaptureRHIResource(const FRHICaptureResourc
 	return bSuccess;
 }
 
+bool UNDIMediaCapture::PostInitializeCaptureMediaTexture(UMediaTexture* InMediaTexture)
+{
+	const bool bSuccess = StartNewCapture();
+	if (bSuccess)
+	{
+		UE_LOG(LogNDIMedia, Log, TEXT("Media Capture Started: Media Texture (%d x %d)."), 
+			InMediaTexture->GetWidth(), InMediaTexture->GetHeight());
+	}
+	return bSuccess;
+}
+
 bool UNDIMediaCapture::UpdateAudioDeviceImpl(const FAudioDeviceHandle& InAudioDeviceHandle)
 {
 	FScopeLock ScopeLock(&CaptureInstanceCriticalSection);
diff --git a/Engine/Plugins/Media/NDIMedia/Source/NDIMedia/Public/NDIMediaCapture.h b/Engine/Plugins/Media/NDIMedia/Source/NDIMedia/Public/NDIMediaCapture.h
--- a/Engine/Plugins/Media/NDIMedia/Source/NDIMedia/Public/NDIMediaCapture.h
+++ b/Engine/Plugins/Media/NDIMedia/Source/NDIMedia/Public/NDIMediaCapture.h
@@ -26,6 +26,7 @@ protected:
 	virtual bool PostInitializeCaptureViewport(TSharedPtr<FSceneViewport>& InSceneViewport) override;
 	virtual bool PostInitializeCaptureRenderTarget(UTextureRenderTarget2D* InRenderTarget) override;
 	virtual bool PostInitializeCaptureRHIResource(const FRHICaptureResourceDescription& InResourceDescription) override;
+	virtual bool PostInitializeCaptureMediaTexture(UMediaTexture* InMediaTexture) override;
 	virtual bool UpdateAudioDeviceImpl(const FAudioDeviceHandle& InAudioDeviceHandle) override;
 	virtual void StopCaptureImpl(bool bAllowPendingFrameToBeProcess) override;
 	virtual bool SupportsAnyThreadCapture() const override;
