Media Profile: Changed capture options property behavior so that default capture settings are used when going to or from media texture captures instead of preserving settings.

#Synced-CL 51845641

diff --git a/Engine/Plugins/Media/MediaFrameworkUtilities/Source/MediaFrameworkUtilitiesEditor/Private/AssetEditor/MediaProfileMediaItemDetailCustomization.cpp b/Engine/Plugins/Media/MediaFrameworkUtilities/Source/MediaFrameworkUtilitiesEditor/Private/AssetEditor/MediaProfileMediaItemDetailCustomization.cpp
--- a/Engine/Plugins/Media/MediaFrameworkUtilities/Source/MediaFrameworkUtilitiesEditor/Private/AssetEditor/MediaProfileMediaItemDetailCustomization.cpp
+++ b/Engine/Plugins/Media/MediaFrameworkUtilities/Source/MediaFrameworkUtilitiesEditor/Private/AssetEditor/MediaProfileMediaItemDetailCustomization.cpp
@@ -651,7 +651,7 @@ private:
 			}
 			
 			// Now, if the object has an existing capture config, we want to copy that config over to the current viewport config, so find it
-			MediaCaptureSettings->CurrentViewportMediaOutput.CaptureOptions = FindExistingMediaCaptureOptions(MediaOutput, MediaCaptureSettings);
+			MediaCaptureSettings->CurrentViewportMediaOutput.CaptureOptions = FindExistingMediaCaptureOptions(MediaOutput, MediaCaptureSettings, FMediaCaptureOptions());
 
 			// We expect there to be only one capture configuration per media output, so clear
 			// any existing capture configurations in the media capture settings before adding the new one
@@ -694,7 +694,7 @@ private:
 			NewViewportOutputInfo.MediaOutput = MediaOutput;
 			
 			// Now, if the object has an existing render target capture config, we want to copy that config over to the new viewport config, so find it
-			NewViewportOutputInfo.CaptureOptions = FindExistingMediaCaptureOptions(MediaOutput, MediaCaptureSettings);
+			NewViewportOutputInfo.CaptureOptions = FindExistingMediaCaptureOptions(MediaOutput, MediaCaptureSettings, NewViewportOutputInfo.CaptureOptions);
 
 			RemoveExistingMediaCaptures(MediaOutput, MediaCaptureSettings);
 			MediaCaptureSettings->ViewportCaptures.Add(NewViewportOutputInfo);
@@ -734,7 +734,7 @@ private:
 			NewRenderTargetOutputInfo.MediaOutput = MediaOutput;
 			
 			// Now, if the object has an existing render target capture config, we want to copy that config over to the new viewport config, so find it
-			NewRenderTargetOutputInfo.CaptureOptions = FindExistingMediaCaptureOptions(MediaOutput, MediaCaptureSettings);
+			NewRenderTargetOutputInfo.CaptureOptions = FindExistingMediaCaptureOptions(MediaOutput, MediaCaptureSettings, NewRenderTargetOutputInfo.CaptureOptions);
 
 			RemoveExistingMediaCaptures(MediaOutput, MediaCaptureSettings);
 			MediaCaptureSettings->RenderTargetCaptures.Add(NewRenderTargetOutputInfo);
@@ -773,14 +773,8 @@ private:
 			FMediaFrameworkCaptureMediaTextureOutputInfo NewMediaTextureOutputInfo;
 			NewMediaTextureOutputInfo.MediaOutput = MediaOutput;
 			
-			// Now, if the object has an existing render target capture config, we want to copy that config over to the new viewport config, so find it
-			NewMediaTextureOutputInfo.CaptureOptions = FindExistingMediaCaptureOptions(MediaOutput, MediaCaptureSettings);
-
-			if (NewMediaTextureOutputInfo.CaptureOptions.ResizeMethod == EMediaCaptureResizeMethod::ResizeSource)
-			{
-				// ResizeSource is not supported for media textures, so revert to None
-				NewMediaTextureOutputInfo.CaptureOptions.ResizeMethod = EMediaCaptureResizeMethod::None;
-			}
+			// Unlike other methods, don't copy capture options from other capture types to media texture, as the there are a specific
+			// set of default options for media textures that are the ideal settings.
 			
 			RemoveExistingMediaCaptures(MediaOutput, MediaCaptureSettings);
 			MediaCaptureSettings->MediaTextureCaptures.Add(NewMediaTextureOutputInfo);
@@ -844,7 +838,7 @@ private:
 		return LOCTEXT("MediaCaptureMethodNotConfigured", "Not Configured");
 	}
 	
-	FMediaCaptureOptions FindExistingMediaCaptureOptions(UMediaOutput* InMediaOutput, UMediaProfileEditorCaptureSettings* MediaCaptureSettings) const
+	FMediaCaptureOptions FindExistingMediaCaptureOptions(UMediaOutput* InMediaOutput, UMediaProfileEditorCaptureSettings* MediaCaptureSettings, const FMediaCaptureOptions& DefaultOptions) const
 	{
 		using namespace UE::MediaFrameworkWorldSettings::Helpers;
 		
@@ -866,13 +860,9 @@ private:
 			return RenderTargetCapture->CaptureOptions;
 		}
 		
-		if (FMediaFrameworkCaptureMediaTextureOutputInfo* MediaTextureCapture =
-			FindFirstOutputInfoForMediaOutput<FMediaFrameworkCaptureMediaTextureOutputInfo>(MediaCaptureSettings, InMediaOutput))
-		{
-			return MediaTextureCapture->CaptureOptions;
-		}
-		
-		return FMediaCaptureOptions();
+		// We don't transfer options from the media texture capture method to other methods as many of the capture options and defaults different for media textures.
+		// If no capture options are found, return the default options
+		return DefaultOptions;
 	}
 
 	void RemoveExistingMediaCaptures(UMediaOutput* InMediaOutput, UMediaProfileEditorCaptureSettings* MediaCaptureSettings)
