Media Profile: Added ability to rotate and flip the media source display in the media profile editor

#virtualized
#Synced-CL 51970848

diff --git a/Engine/Plugins/Media/MediaFrameworkUtilities/Source/MediaFrameworkUtilities/Private/Profile/MediaProfilePlaybackManager.cpp b/Engine/Plugins/Media/MediaFrameworkUtilities/Source/MediaFrameworkUtilities/Private/Profile/MediaProfilePlaybackManager.cpp
--- a/Engine/Plugins/Media/MediaFrameworkUtilities/Source/MediaFrameworkUtilities/Private/Profile/MediaProfilePlaybackManager.cpp
+++ b/Engine/Plugins/Media/MediaFrameworkUtilities/Source/MediaFrameworkUtilities/Private/Profile/MediaProfilePlaybackManager.cpp
@@ -997,7 +997,7 @@ UMediaCapture* UMediaProfilePlaybackManager::OpenRenderTargetOutput(UMediaOutput
 	return OpenRenderTargetOutputFromIndex(MediaOutputIndex, InRenderTarget, InCaptureOptions, bAutoRestartCapture);
 }
 
-UMediaCapture* UMediaProfilePlaybackManager::OpenMediaTextureOutputFromIndex(int32 InMediaOutputIndex, UMediaTexture* InMediaTexture, const FMediaCaptureOptions& InCaptureOptions, bool bAutoRestartCapture)
+UMediaCapture* UMediaProfilePlaybackManager::OpenMediaTextureOutputFromIndex(int32 InMediaOutputIndex, UMediaTexture* InMediaTexture, const FMediaCaptureOptions& InCaptureOptions, const FMediaCaptureTransform& InTransform, bool bAutoRestartCapture)
 {
 	UMediaProfile* MediaProfile = GetOwningMediaProfile();
 	if (!MediaProfile)
@@ -1024,12 +1024,12 @@ UMediaCapture* UMediaProfilePlaybackManager::OpenMediaTextureOutputFromIndex(int
 		return MediaCapture;
 	};
 
-	if (MediaCapture->CaptureMediaTexture(InMediaTexture, InCaptureOptions))
+	if (MediaCapture->CaptureMediaTexture(InMediaTexture, InCaptureOptions, InTransform.GetRotationValue(true), InTransform.GetScaleValue()))
 	{
 #if WITH_EDITOR
 		if (bAutoRestartCapture)
 		{
-			MediaOutput->OnOutputModified().AddUObject(this, &UMediaProfilePlaybackManager::RestartMediaTextureOutput, InMediaTexture, InCaptureOptions, bAutoRestartCapture);
+			MediaOutput->OnOutputModified().AddUObject(this, &UMediaProfilePlaybackManager::RestartMediaTextureOutput, InMediaTexture, InCaptureOptions, InTransform, bAutoRestartCapture);
 		}
 #endif
 		
@@ -1043,7 +1043,7 @@ UMediaCapture* UMediaProfilePlaybackManager::OpenMediaTextureOutputFromIndex(int
 	return MediaCapture;
 }
 
-UMediaCapture* UMediaProfilePlaybackManager::OpenMediaTextureOutput(UMediaOutput* InMediaOutput, UMediaTexture* InMediaTexture, const FMediaCaptureOptions& InCaptureOptions, bool bAutoRestartCapture)
+UMediaCapture* UMediaProfilePlaybackManager::OpenMediaTextureOutput(UMediaOutput* InMediaOutput, UMediaTexture* InMediaTexture, const FMediaCaptureOptions& InCaptureOptions, const FMediaCaptureTransform& InTransform, bool bAutoRestartCapture)
 {
 	UMediaProfile* MediaProfile = GetOwningMediaProfile();
 	if (!MediaProfile)
@@ -1052,7 +1052,7 @@ UMediaCapture* UMediaProfilePlaybackManager::OpenMediaTextureOutput(UMediaOutput
 	}
 	
 	int32 MediaOutputIndex = MediaProfile->FindMediaOutputIndex(InMediaOutput);
-	return OpenMediaTextureOutputFromIndex(MediaOutputIndex, InMediaTexture, InCaptureOptions, bAutoRestartCapture);
+	return OpenMediaTextureOutputFromIndex(MediaOutputIndex, InMediaTexture, InCaptureOptions, InTransform, bAutoRestartCapture);
 }
 
 void UMediaProfilePlaybackManager::RestartActiveViewportOutput(UMediaOutput* InMediaOutput, FMediaCaptureOptions InCaptureOptions, bool bAutoRestartCapture)
@@ -1102,18 +1102,18 @@ void UMediaProfilePlaybackManager::RestartRenderTargetOutput(UMediaOutput* InMed
 	QueueRestartCapture(InMediaOutput, Callback);
 }
 
-void UMediaProfilePlaybackManager::RestartMediaTextureOutput(UMediaOutput* InMediaOutput, UMediaTexture* InMediaTexture, FMediaCaptureOptions InCaptureOptions, bool bAutoRestartCapture)
+void UMediaProfilePlaybackManager::RestartMediaTextureOutput(UMediaOutput* InMediaOutput, UMediaTexture* InMediaTexture, FMediaCaptureOptions InCaptureOptions, FMediaCaptureTransform InTransform, bool bAutoRestartCapture)
 {
 	TWeakObjectPtr<UMediaTexture> WeakMediaTexture = InMediaTexture;
 	TWeakObjectPtr<UMediaOutput> WeakMediaOutput = InMediaOutput;
-	auto Callback = [WeakMediaTexture, WeakMediaOutput, InCaptureOptions, bAutoRestartCapture](UMediaProfilePlaybackManager* Manager)
+	auto Callback = [WeakMediaTexture, WeakMediaOutput, InCaptureOptions, InTransform, bAutoRestartCapture](UMediaProfilePlaybackManager* Manager)
 	{
 		TStrongObjectPtr<UMediaTexture> MediaTexture = WeakMediaTexture.Pin();
 		TStrongObjectPtr<UMediaOutput> MediaOutput = WeakMediaOutput.Pin();
 
 		if (Manager && MediaTexture && MediaOutput)
 		{
-			Manager->OpenMediaTextureOutput(MediaOutput.Get(), MediaTexture.Get(), InCaptureOptions, bAutoRestartCapture);
+			Manager->OpenMediaTextureOutput(MediaOutput.Get(), MediaTexture.Get(), InCaptureOptions, InTransform, bAutoRestartCapture);
 		}
 	};
 	QueueRestartCapture(InMediaOutput, Callback);
diff --git a/Engine/Plugins/Media/MediaFrameworkUtilities/Source/MediaFrameworkUtilities/Public/Profile/MediaProfilePlaybackManager.h b/Engine/Plugins/Media/MediaFrameworkUtilities/Source/MediaFrameworkUtilities/Public/Profile/MediaProfilePlaybackManager.h
--- a/Engine/Plugins/Media/MediaFrameworkUtilities/Source/MediaFrameworkUtilities/Public/Profile/MediaProfilePlaybackManager.h
+++ b/Engine/Plugins/Media/MediaFrameworkUtilities/Source/MediaFrameworkUtilities/Public/Profile/MediaProfilePlaybackManager.h
@@ -5,11 +5,13 @@
 #include "UObject/Object.h"
 #include "Engine/EngineBaseTypes.h"
 #include "Engine/TimerHandle.h"
+#include "UObject/SoftObjectPtr.h"
 
 #include "MediaProfilePlaybackManager.generated.h"
 
 enum class EMediaCaptureState : uint8;
 struct FMediaCaptureOptions;
+struct FMediaCaptureTransform;
 class FSceneViewport;
 class FViewportClient;
 class UMediaOutput;
@@ -239,20 +241,22 @@ public:
 	 * @param InMediaOutputIndex The index of the media output in the media profile
 	 * @param InMediaTexture The media texture to capture
 	 * @param InCaptureOptions The capture options to capture with
+	 * @param InTransform The transform to apply to the media texture on output
 	 * @param bAutoRestartCapture Indicates that this capture will automatically restart if any modifications to the media output are detected
 	 * @return The media capture object, if the capture was started successfully; otherwise, null
 	 */
-	UMediaCapture* OpenMediaTextureOutputFromIndex(int32 InMediaOutputIndex, UMediaTexture* InMediaTexture, const FMediaCaptureOptions& InCaptureOptions, bool bAutoRestartCapture = false);
+	UMediaCapture* OpenMediaTextureOutputFromIndex(int32 InMediaOutputIndex, UMediaTexture* InMediaTexture, const FMediaCaptureOptions& InCaptureOptions, const FMediaCaptureTransform& InTransform, bool bAutoRestartCapture = false);
 	
 	/**
 	 * Starts a capture of a media texture for a specified media output
 	 * @param InMediaOutput The media output in the media profile
 	 * @param InMediaTexture The media texture to capture
 	 * @param InCaptureOptions The capture options to capture with
+	 * @param InTransform The transform to apply to the media texture on output
 	 * @param bAutoRestartCapture Indicates that this capture will automatically restart if any modifications to the media output are detected
 	 * @return The media capture object, if the capture was started successfully; otherwise, null
 	 */
-	UMediaCapture* OpenMediaTextureOutput(UMediaOutput* InMediaOutput, UMediaTexture* InMediaTexture, const FMediaCaptureOptions& InCaptureOptions, bool bAutoRestartCapture = false);
+	UMediaCapture* OpenMediaTextureOutput(UMediaOutput* InMediaOutput, UMediaTexture* InMediaTexture, const FMediaCaptureOptions& InCaptureOptions, const FMediaCaptureTransform& InTransform, bool bAutoRestartCapture = false);
 	
 	 /* 
 	 * Restarts a capture of an existing engine viewport for a specified media output
@@ -284,9 +288,10 @@ public:
 	 * @param InMediaOutput The media output in the media profile
 	 * @param InMediaTexture The media texture to capture
 	 * @param InCaptureOptions The capture options to capture with
+	 * @param InTransform The transform to apply to the media texture on output
 	 * @param bAutoRestartCapture Indicates that this capture will automatically restart if any modifications to the media output are detected
 	 */
-	void RestartMediaTextureOutput(UMediaOutput* InMediaOutput, UMediaTexture* InMediaTexture, FMediaCaptureOptions InCaptureOptions, bool bAutoRestartCapture = false);
+	void RestartMediaTextureOutput(UMediaOutput* InMediaOutput, UMediaTexture* InMediaTexture, FMediaCaptureOptions InCaptureOptions, FMediaCaptureTransform InTransform, bool bAutoRestartCapture = false);
 	
 	/**
 	 * Gets whether there is an active capture for the specified media output index
diff --git a/Engine/Plugins/Media/MediaFrameworkUtilities/Source/MediaFrameworkUtilitiesEditor/Private/AssetEditor/MediaProfileEditor.cpp b/Engine/Plugins/Media/MediaFrameworkUtilities/Source/MediaFrameworkUtilitiesEditor/Private/AssetEditor/MediaProfileEditor.cpp
--- a/Engine/Plugins/Media/MediaFrameworkUtilities/Source/MediaFrameworkUtilitiesEditor/Private/AssetEditor/MediaProfileEditor.cpp
+++ b/Engine/Plugins/Media/MediaFrameworkUtilities/Source/MediaFrameworkUtilitiesEditor/Private/AssetEditor/MediaProfileEditor.cpp
@@ -514,7 +514,7 @@ void FMediaProfileEditor::RestartActiveMediaCaptures(UMediaOutput* InMediaOutput
 	ForEachOutputInfoForMediaOutput<FMediaFrameworkCaptureMediaTextureOutputInfo>(CaptureSettings, InMediaOutput, 
 		[&PlaybackManager, bAutoRestart = CaptureSettings->bAutoRestartCaptureOnChange](const FMediaFrameworkCaptureMediaTextureOutputInfo& OutputInfo)
 		{
-			PlaybackManager->RestartMediaTextureOutput(OutputInfo.MediaOutput, OutputInfo.MediaTexture, OutputInfo.CaptureOptions, bAutoRestart);
+			PlaybackManager->RestartMediaTextureOutput(OutputInfo.MediaOutput, OutputInfo.MediaTexture, OutputInfo.CaptureOptions, OutputInfo.Transform, bAutoRestart);
 		});
 }
 
diff --git a/Engine/Plugins/Media/MediaFrameworkUtilities/Source/MediaFrameworkUtilitiesEditor/Private/AssetEditor/MediaProfileEditorUserSettings.h b/Engine/Plugins/Media/MediaFrameworkUtilities/Source/MediaFrameworkUtilitiesEditor/Private/AssetEditor/MediaProfileEditorUserSettings.h
--- a/Engine/Plugins/Media/MediaFrameworkUtilities/Source/MediaFrameworkUtilitiesEditor/Private/AssetEditor/MediaProfileEditorUserSettings.h
+++ b/Engine/Plugins/Media/MediaFrameworkUtilities/Source/MediaFrameworkUtilitiesEditor/Private/AssetEditor/MediaProfileEditorUserSettings.h
@@ -7,6 +7,16 @@
 
 #include "MediaProfileEditorUserSettings.generated.h"
 
+/** Rotation of the displayed media image in the media profile editor */
+UENUM()
+enum class EMediaImageRotation
+{
+	None UMETA(DisplayName = "0°"),
+	CW90 UMETA(DisplayName = "90°"),
+	CCW90 UMETA(DisplayName = "-90°"),
+	CW180 UMETA(DisplayName = "180°")
+};
+
 /** Settings for a media source/output in the media profile editor */
 USTRUCT()
 struct FMediaProfileEditorPerMediaItemSettings
@@ -25,6 +35,18 @@ public:
 	/** The saved draw checkerboard flag applied to the media source/output */
 	UPROPERTY(Config)
 	bool bDrawAlphaBlendedCheckerboard = false;
+	
+	/** The rotation of the media image in the viewport */
+	UPROPERTY(Config)
+	EMediaImageRotation Rotation = EMediaImageRotation::None;
+	
+	/** Whether the media image is flipped horizontally */
+	UPROPERTY(Config)
+	bool bFlipHorizontal = false;
+	
+	/** Whether the media image is flipped vertically */
+	UPROPERTY(Config)
+	bool bFlipVertical = false;
 };
 
 /** Config class to store editor user settings for the media profile editor */
diff --git a/Engine/Plugins/Media/MediaFrameworkUtilities/Source/MediaFrameworkUtilitiesEditor/Private/AssetEditor/MediaProfileMediaItemDetailCustomization.cpp b/Engine/Plugins/Media/MediaFrameworkUtilities/Source/MediaFrameworkUtilitiesEditor/Private/AssetEditor/MediaProfileMediaItemDetailCustomization.cpp
--- a/Engine/Plugins/Media/MediaFrameworkUtilities/Source/MediaFrameworkUtilitiesEditor/Private/AssetEditor/MediaProfileMediaItemDetailCustomization.cpp
+++ b/Engine/Plugins/Media/MediaFrameworkUtilities/Source/MediaFrameworkUtilitiesEditor/Private/AssetEditor/MediaProfileMediaItemDetailCustomization.cpp
@@ -15,6 +15,7 @@
 #include "Internationalization/Culture.h"
 #include "Profile/MediaProfilePlaybackManager.h"
 #include "Widgets/SToolTip.h"
+#include "Widgets/Input/SSegmentedControl.h"
 #include "Widgets/Input/SComboBox.h"
 
 #define LOCTEXT_NAMESPACE "FMediaProfileMediaItemDetailCustomization"
@@ -942,6 +943,11 @@ protected:
 	/** Raised when any of the capture info properties have changed */
 	virtual void PropertyValueChange(const FPropertyChangedEvent& PropertyChangedEvent)
 	{
+		if (!bBroadcastPropertyChanged)
+		{
+			return;
+		}
+		
 		if (UMediaProfileEditorCaptureSettings* MediaCaptureSettings = FMediaProfileEditor::GetMediaFrameworkCaptureSettings())
 		{
 			MediaCaptureSettings->SaveConfig();
@@ -1016,6 +1022,9 @@ protected:
 protected:
 	/** Weak pointer to the  media profile editor */
 	TWeakPtr<FMediaProfileEditor> MediaProfileEditor;
+	
+	/** Indicates that ObjectPropertyChanged event should be broadcast when a child property value has changed */
+	bool bBroadcastPropertyChanged = true;
 };
 
 class FMediaFrameworkViewportOutputInfoCustomization : public FMediaFrameworkBaseOutputInfoCustomization
@@ -1081,15 +1090,63 @@ protected:
 			.OnMediaSourceSelected(this, &FMediaFrameworkMediaTextureOutputInfoCustomization::OnMediaSourceSelected)
 		];
 
+		ResizeMethodProperty = PropertyHandle->GetChildHandle(GET_MEMBER_NAME_CHECKED(FMediaCaptureOptions, ResizeMethod));
+		RotationProperty = PropertyHandle->GetChildHandle(GET_MEMBER_NAME_CHECKED(FMediaCaptureTransform, Rotation));
+		RotationProperty->SetOnPropertyValueChangedWithData(TDelegate<void(const FPropertyChangedEvent&)>::CreateSP(this, &FMediaFrameworkMediaTextureOutputInfoCustomization::PropertyValueChange));
+		
+		ChildBuilder.AddCustomRow(LOCTEXT("RotationFilterString", "Rotation"))
+			.PropertyHandleList({ RotationProperty })
+			.NameContent()
+			[
+				RotationProperty->CreatePropertyNameWidget()
+			]
+			.ValueContent()
+			[
+				CreateRotationWidget()
+			];
+	
+		FlipHorizontalProperty = PropertyHandle->GetChildHandle(GET_MEMBER_NAME_CHECKED(FMediaCaptureTransform, bFlipHorizontal));
+		FlipHorizontalProperty->SetOnPropertyValueChangedWithData(TDelegate<void(const FPropertyChangedEvent&)>::CreateSP(this, &FMediaFrameworkMediaTextureOutputInfoCustomization::PropertyValueChange));
+		
+		FlipVerticalProperty = PropertyHandle->GetChildHandle(GET_MEMBER_NAME_CHECKED(FMediaCaptureTransform, bFlipVertical));
+		FlipVerticalProperty->SetOnPropertyValueChangedWithData(TDelegate<void(const FPropertyChangedEvent&)>::CreateSP(this, &FMediaFrameworkMediaTextureOutputInfoCustomization::PropertyValueChange));
+		
+		ChildBuilder.AddCustomRow(LOCTEXT("FlipFilterString", "Flip"))
+			.NameContent()
+			[
+				SNew(STextBlock)
+				.Text(LOCTEXT("FlipName", "Flip"))
+				.Font(IDetailLayoutBuilder::GetDetailFont())
+			]
+			.ValueContent()
+			[
+				CreateFlipWidget()
+			]
+			.OverrideResetToDefault(FResetToDefaultOverride::Create(
+				TAttribute<bool>::CreateSP(this, &FMediaFrameworkMediaTextureOutputInfoCustomization::CanResetFlipToDefault),
+				FSimpleDelegate::CreateSP(this, &FMediaFrameworkMediaTextureOutputInfoCustomization::ResetFlipToDefault)
+			));
+
 		FMediaFrameworkBaseOutputInfoCustomization::CustomizeChildren(PropertyHandle, ChildBuilder, CustomizationUtils);
 	}
 
 	virtual TArray<EMediaCaptureResizeMethod> GetRestrictedResizeMethods() const override
 	{
-		return 
+		TArray<EMediaCaptureResizeMethod> DisallowedResizeMethods = { EMediaCaptureResizeMethod::ResizeSource };
+		
+		if (RotationProperty.IsValid() && RotationProperty->IsValidHandle())
 		{
-			EMediaCaptureResizeMethod::ResizeSource
-		};
+			uint8 RotationValue;
+			if (RotationProperty->GetValue(RotationValue) == FPropertyAccess::Success)
+			{
+				if ((EMediaCaptureRotation)RotationValue != EMediaCaptureRotation::None)
+				{
+					DisallowedResizeMethods.Add(EMediaCaptureResizeMethod::None);
+				}
+			}
+		}
+		
+		return DisallowedResizeMethods;
 	}
 	
 private:
@@ -1106,6 +1163,203 @@ private:
 			InMediaProfile->GetPlaybackManager()->OpenSourceFromIndex(InMediaSourceIndex, PinnedMediaProfileEditor.Get());
 		}
 	}
+	
+	TSharedRef<SWidget> CreateRotationWidget()
+	{
+		return SNew(SSegmentedControl<EMediaCaptureRotation>)
+			.OnValueChanged(this, &FMediaFrameworkMediaTextureOutputInfoCustomization::SetMediaCaptureRotation)
+			.Value(this, &FMediaFrameworkMediaTextureOutputInfoCustomization::GetMediaCaptureRotation)
+			
+			+SSegmentedControl<EMediaCaptureRotation>::Slot(EMediaCaptureRotation::CCW90)
+			.Text(StaticEnum<EMediaCaptureRotation>()->GetDisplayNameTextByValue((int64)EMediaCaptureRotation::CCW90))
+			
+			+SSegmentedControl<EMediaCaptureRotation>::Slot(EMediaCaptureRotation::None)
+			.Text(StaticEnum<EMediaCaptureRotation>()->GetDisplayNameTextByValue((int64)EMediaCaptureRotation::None))
+			
+			+SSegmentedControl<EMediaCaptureRotation>::Slot(EMediaCaptureRotation::CW90)
+			.Text(StaticEnum<EMediaCaptureRotation>()->GetDisplayNameTextByValue((int64)EMediaCaptureRotation::CW90))
+			
+			+SSegmentedControl<EMediaCaptureRotation>::Slot(EMediaCaptureRotation::CW180)
+			.Text(StaticEnum<EMediaCaptureRotation>()->GetDisplayNameTextByValue((int64)EMediaCaptureRotation::CW180));
+	}
+
+	EMediaCaptureRotation GetMediaCaptureRotation() const
+	{
+		if (!RotationProperty.IsValid() || !RotationProperty->IsValidHandle())
+		{
+			return EMediaCaptureRotation::None;
+		}
+		
+		uint8 EnumValue;
+		FPropertyAccess::Result Result = RotationProperty->GetValue(EnumValue);
+		if (RotationProperty->GetValue(EnumValue) == FPropertyAccess::Success)
+		{
+			return (EMediaCaptureRotation)EnumValue;
+		}
+		
+		return EMediaCaptureRotation::None;
+	}
+
+	void SetMediaCaptureRotation(EMediaCaptureRotation InNewRotation)
+	{
+		if (!RotationProperty.IsValid() || !RotationProperty->IsValidHandle())
+		{
+			return;
+		}
+
+		RotationProperty->SetValue((uint8)InNewRotation);
+		
+		if (InNewRotation != EMediaCaptureRotation::None)
+		{
+			if (ResizeMethodProperty.IsValid() && ResizeMethodProperty->IsValidHandle())
+			{
+				ResizeMethodProperty->SetValue((uint8)EMediaCaptureResizeMethod::ResizeInRenderPass);
+			}
+		}
+	}
+
+	TSharedRef<SWidget> CreateFlipWidget()
+	{
+		return SNew(SHorizontalBox)
+			
+			+SHorizontalBox::Slot()
+			.AutoWidth()
+			.Padding(0.0f, 0.0f, 2.0f, 0.0f)
+			[
+				SNew(SCheckBox)
+				.Style(FAppStyle::Get(), "DetailsView.SectionButton")
+				.Padding(FMargin(6.0f, 4.0f))
+				.IsChecked(this, &FMediaFrameworkMediaTextureOutputInfoCustomization::IsFlipped, FlipHorizontalProperty)
+				.OnCheckStateChanged(this, &FMediaFrameworkMediaTextureOutputInfoCustomization::SetFlipped, FlipHorizontalProperty)
+				[
+					SNew(SHorizontalBox)
+
+					+SHorizontalBox::Slot()
+					.AutoWidth()
+					.Padding(0.0f, 0.0f, 2.0f, 0.0f)
+					.VAlign(VAlign_Center)
+					[
+						SNew(SImage)
+						.Image(FAppStyle::GetBrush("Icons.FlipHorizontal"))
+						.ColorAndOpacity(FSlateColor::UseForeground())
+					]
+
+					+SHorizontalBox::Slot()
+					.VAlign(VAlign_Center)
+					.AutoWidth()
+					[
+						SNew(STextBlock)
+						.Text(LOCTEXT("FlipHorizontalLabel", "Flip Horizontal"))
+						.Font(IDetailLayoutBuilder::GetDetailFont())
+					]
+				]
+			]
+			
+			+SHorizontalBox::Slot()
+			.AutoWidth()
+			[
+				SNew(SCheckBox)
+				.Style(FAppStyle::Get(), "DetailsView.SectionButton")
+				.Padding(FMargin(6.0f, 4.0f))
+				.IsChecked(this, &FMediaFrameworkMediaTextureOutputInfoCustomization::IsFlipped, FlipVerticalProperty)
+				.OnCheckStateChanged(this, &FMediaFrameworkMediaTextureOutputInfoCustomization::SetFlipped, FlipVerticalProperty)
+				[
+					SNew(SHorizontalBox)
+
+					+SHorizontalBox::Slot()
+					.AutoWidth()
+					.Padding(0.0f, 0.0f, 2.0f, 0.0f)
+					.VAlign(VAlign_Center)
+					[
+						SNew(SImage)
+						.Image(FAppStyle::GetBrush("Icons.FlipVertical"))
+						.ColorAndOpacity(FSlateColor::UseForeground())
+					]
+
+					+SHorizontalBox::Slot()
+					.VAlign(VAlign_Center)
+					.AutoWidth()
+					[
+						SNew(STextBlock)
+						.Text(LOCTEXT("FlipVerticalLabel", "Flip Vertical"))
+						.Font(IDetailLayoutBuilder::GetDetailFont())
+					]
+				]
+			];
+	}
+
+	ECheckBoxState IsFlipped(TSharedPtr<IPropertyHandle> InFlippedProperty) const
+	{
+		if (!InFlippedProperty.IsValid() || !InFlippedProperty->IsValidHandle())
+		{
+			return ECheckBoxState::Unchecked;
+		}
+		
+		bool bFlippedValue;
+		if (InFlippedProperty->GetValue(bFlippedValue) == FPropertyAccess::Success)
+		{
+			return bFlippedValue ? ECheckBoxState::Checked : ECheckBoxState::Unchecked;
+		}
+
+		return ECheckBoxState::Unchecked;
+	}
+
+	void SetFlipped(ECheckBoxState CheckBoxState, TSharedPtr<IPropertyHandle> InFlippedProperty)
+	{
+		if (!InFlippedProperty.IsValid() || !InFlippedProperty->IsValidHandle())
+		{
+			return;
+		}
+		
+		InFlippedProperty->SetValue(CheckBoxState == ECheckBoxState::Checked);
+	}
+	
+	bool CanResetFlipToDefault() const
+	{
+		bool bFlipHorizontalValue = false;
+		bool bFlipVerticalValue = false;
+		
+		if (FlipHorizontalProperty.IsValid() && FlipHorizontalProperty->IsValidHandle())
+		{
+			FlipHorizontalProperty->GetValue(bFlipHorizontalValue);
+		}
+		
+		if (FlipVerticalProperty.IsValid() && FlipVerticalProperty->IsValidHandle())
+		{
+			FlipVerticalProperty->GetValue(bFlipVerticalValue);
+		}
+		
+		return bFlipHorizontalValue || bFlipVerticalValue;
+	}
+	
+	void ResetFlipToDefault()
+	{
+		{
+			// Delay broadcasting the ObjectPropertyChanged event until both properties have been set
+			TGuardValue<bool> Guard(bBroadcastPropertyChanged, false);
+			if (FlipHorizontalProperty.IsValid() && FlipHorizontalProperty->IsValidHandle())
+			{
+				FlipHorizontalProperty->SetValue(false, EPropertyValueSetFlags::ResetToDefault);
+			}
+		
+			if (FlipVerticalProperty.IsValid() && FlipVerticalProperty->IsValidHandle())
+			{
+				FlipVerticalProperty->SetValue(false, EPropertyValueSetFlags::ResetToDefault);
+			}
+		}
+
+		FPropertyChangedEvent FlipHorizontalEvent(FindFProperty<FProperty>(FMediaCaptureTransform::StaticStruct(), GET_MEMBER_NAME_CHECKED(FMediaCaptureTransform, bFlipHorizontal)));
+		PropertyValueChange(FlipHorizontalEvent);
+		
+		FPropertyChangedEvent FlipVerticalEvent(FindFProperty<FProperty>(FMediaCaptureTransform::StaticStruct(), GET_MEMBER_NAME_CHECKED(FMediaCaptureTransform, bFlipVertical)));
+		PropertyValueChange(FlipVerticalEvent);
+	}
+	
+private:
+	TSharedPtr<IPropertyHandle> ResizeMethodProperty;
+	TSharedPtr<IPropertyHandle> RotationProperty;
+	TSharedPtr<IPropertyHandle> FlipHorizontalProperty;
+	TSharedPtr<IPropertyHandle> FlipVerticalProperty;
 };
 
 FMediaProfileMediaOutputDetailCustomization::~FMediaProfileMediaOutputDetailCustomization()
@@ -1445,6 +1699,52 @@ void FMediaProfileMediaOutputDetailCustomization::OnObjectPropertyChanged(UObjec
 		{
 			CachedDetailBuilder.Pin()->ForceRefreshDetails();
 		}
+		
+		static TSet<FName> RefreshableProperties = 
+		{
+			GET_MEMBER_NAME_CHECKED(FMediaCaptureTransform, Rotation),
+			GET_MEMBER_NAME_CHECKED(FMediaCaptureTransform, bFlipHorizontal),
+			GET_MEMBER_NAME_CHECKED(FMediaCaptureTransform, bFlipVertical),
+			GET_MEMBER_NAME_CHECKED(FMediaCaptureOptions, Crop),
+			GET_MEMBER_NAME_CHECKED(FMediaCaptureOptions, ResizeMethod),
+			GET_MEMBER_NAME_CHECKED(FMediaCaptureOptions, bConvertToDesiredPixelFormat),
+			GET_MEMBER_NAME_CHECKED(FMediaCaptureOptions, bForceAlphaToOneOnConversion),
+			GET_MEMBER_NAME_CHECKED(FMediaCaptureOptions, bApplyLinearToSRGBConversion)
+		};
+		
+		if (RefreshableProperties.Contains(PropertyName))
+		{
+			if (TSharedPtr<FMediaProfileEditor> PinnedMediaProfileEditor = MediaProfileEditor.Pin())
+			{
+				TArray<TWeakObjectPtr<UObject>> Objects;
+				if (TSharedPtr<IDetailLayoutBuilder> DetailBuilder = CachedDetailBuilder.Pin())
+				{
+					CachedDetailBuilder.Pin()->GetObjectsBeingCustomized(Objects);
+				}
+
+				for (const TWeakObjectPtr<UObject>& Object : Objects)
+				{
+					TStrongObjectPtr<UObject> PinnedObject = Object.Pin();
+					if (!PinnedObject.IsValid())
+					{
+						continue;
+					}
+
+					if (!PinnedObject->IsA<UMediaOutput>())
+					{
+						continue;
+					}
+
+					if (PinnedObject->IsA<UDummyMediaOutput>())
+					{
+						continue;
+					}
+
+					UMediaOutput* MediaOutput = Cast<UMediaOutput>(PinnedObject.Get());
+					PinnedMediaProfileEditor->RestartActiveMediaCaptures(MediaOutput);
+				}
+			}
+		}
 	}
 }
 
diff --git a/Engine/Plugins/Media/MediaFrameworkUtilities/Source/MediaFrameworkUtilitiesEditor/Private/AssetEditor/SMaskedMediaImage.cpp b/Engine/Plugins/Media/MediaFrameworkUtilities/Source/MediaFrameworkUtilitiesEditor/Private/AssetEditor/SMaskedMediaImage.cpp
--- a/Engine/Plugins/Media/MediaFrameworkUtilities/Source/MediaFrameworkUtilitiesEditor/Private/AssetEditor/SMaskedMediaImage.cpp
+++ b/Engine/Plugins/Media/MediaFrameworkUtilities/Source/MediaFrameworkUtilitiesEditor/Private/AssetEditor/SMaskedMediaImage.cpp
@@ -16,6 +16,8 @@ void SMaskedMediaImage::Construct(const FArguments& InArgs, UTexture* InMediaIma
 	ChannelMaskAttr = InArgs._ChannelMask;
 	InvertAlphaChannelAttr = InArgs._InvertAlphaChannel;
 	DrawCheckerboardAttr = InArgs._DrawCheckerboard;
+	RotationAttr = InArgs._Rotation;
+	ScaleAttr = InArgs._Scale;
 
 	UMaterial* BaseMaterial = (UMaterial*)StaticLoadObject(UMaterial::StaticClass(), nullptr, TEXT("/MediaFrameworkUtilities/Editor/M_MaskedMediaImage.M_MaskedMediaImage"));
 	check(BaseMaterial);
@@ -42,15 +44,29 @@ void SMaskedMediaImage::Tick(const FGeometry& AllottedGeometry, const double InC
 {
 	SCompoundWidget::Tick(AllottedGeometry, InCurrentTime, InDeltaTime);
 
+	const FVector2f Scale = ScaleAttr.Get(FVector2f::UnitVector);
+	const float Rotation = RotationAttr.Get(0.0f);
+
 	if (ImageSizeAttr.IsSet())
 	{
-		FVector2D Size = ImageSizeAttr.Get();
-		MaterialBrush->ImageSize.X = static_cast<float>(Size.X);
-		MaterialBrush->ImageSize.Y = static_cast<float>(Size.Y);
+		const FVector2D Size = ImageSizeAttr.Get();
+		const FVector2D RotatedSize = FVector2D(
+			Size.X * FMath::Abs(FMath::Cos(Rotation)) + Size.Y * FMath::Abs(FMath::Sin(Rotation)),
+			Size.X * FMath::Abs(FMath::Sin(Rotation)) + Size.Y * FMath::Abs(FMath::Cos(Rotation)));
+		
+		MaterialBrush->ImageSize.X = static_cast<float>(RotatedSize.X);
+		MaterialBrush->ImageSize.Y = static_cast<float>(RotatedSize.Y);
 	}
 
 	const int8 ChannelMask = (int8)ChannelMaskAttr.Get(EColorChannelMask::All);
 	Material->SetScalarParameterValue(TEXT("ChannelMask"), ChannelMask + 1);
 	Material->SetScalarParameterValue(TEXT("bInvertAlphaChannelMask"), InvertAlphaChannelAttr.Get(false) ? 1.0f : 0.0f);
 	Material->SetScalarParameterValue(TEXT("bDrawCheckerboard"), DrawCheckerboardAttr.Get(false) ? 1.0f : 0.0f);
+	
+	const FVector4 RotScale = FVector4(
+		Scale.X * FMath::Cos(Rotation), -Scale.X * FMath::Sin(Rotation),
+		Scale.Y * FMath::Sin(Rotation), Scale.Y * FMath::Cos(Rotation)
+	);
+	
+	Material->SetVectorParameterValue(TEXT("RotScale"), RotScale);
 }
diff --git a/Engine/Plugins/Media/MediaFrameworkUtilities/Source/MediaFrameworkUtilitiesEditor/Private/AssetEditor/SMaskedMediaImage.h b/Engine/Plugins/Media/MediaFrameworkUtilities/Source/MediaFrameworkUtilitiesEditor/Private/AssetEditor/SMaskedMediaImage.h
--- a/Engine/Plugins/Media/MediaFrameworkUtilities/Source/MediaFrameworkUtilitiesEditor/Private/AssetEditor/SMaskedMediaImage.h
+++ b/Engine/Plugins/Media/MediaFrameworkUtilities/Source/MediaFrameworkUtilitiesEditor/Private/AssetEditor/SMaskedMediaImage.h
@@ -20,12 +20,16 @@ public:
 	
 	SLATE_BEGIN_ARGS(SMaskedMediaImage)
 		: _IsAlphaPremultiplied(true)
+		, _Rotation(0.0f)
+		, _Scale(FVector2D(1.0f, 1.0f))
 	{ }
 		SLATE_ATTRIBUTE(FVector2D, ImageSize)
 		SLATE_ATTRIBUTE(EColorChannelMask, ChannelMask)
 		SLATE_ATTRIBUTE(bool, InvertAlphaChannel)
 		SLATE_ATTRIBUTE(bool, DrawCheckerboard)
 		SLATE_ARGUMENT(bool, IsAlphaPremultiplied)
+		SLATE_ATTRIBUTE(float, Rotation)
+		SLATE_ATTRIBUTE(FVector2f, Scale)
 	SLATE_END_ARGS()
 
 	void Construct(const FArguments& InArgs, UTexture* InMediaImageTexture);
@@ -53,4 +57,10 @@ private:
 
 	/** Attribute for drawing the alpha-blended checkerboard background */
 	TAttribute<bool> DrawCheckerboardAttr;
+	
+	/** The 2D rotation of the image */
+	TAttribute<float> RotationAttr;
+	
+	/** The 2D scale of the image */
+	TAttribute<FVector2f> ScaleAttr;
 };
diff --git a/Engine/Plugins/Media/MediaFrameworkUtilities/Source/MediaFrameworkUtilitiesEditor/Private/AssetEditor/SMediaProfileDetailsPanel.cpp b/Engine/Plugins/Media/MediaFrameworkUtilities/Source/MediaFrameworkUtilitiesEditor/Private/AssetEditor/SMediaProfileDetailsPanel.cpp
--- a/Engine/Plugins/Media/MediaFrameworkUtilities/Source/MediaFrameworkUtilitiesEditor/Private/AssetEditor/SMediaProfileDetailsPanel.cpp
+++ b/Engine/Plugins/Media/MediaFrameworkUtilities/Source/MediaFrameworkUtilitiesEditor/Private/AssetEditor/SMediaProfileDetailsPanel.cpp
@@ -313,7 +313,7 @@ void SMediaProfileDetailsPanel::OnCaptureButtonPressed()
 						PinnedMediaProfile->GetPlaybackManager()->OpenSourceFromIndex(MediaSourceIndex, OwningEditor.Pin().Get());
 					}
 					
-					PinnedMediaProfile->GetPlaybackManager()->OpenMediaTextureOutput(MediaOutput, MediaTextureCapture->MediaTexture, MediaTextureCapture->CaptureOptions, MediaCaptureSettings->bAutoRestartCaptureOnChange);
+					PinnedMediaProfile->GetPlaybackManager()->OpenMediaTextureOutput(MediaOutput, MediaTextureCapture->MediaTexture, MediaTextureCapture->CaptureOptions, MediaTextureCapture->Transform, MediaCaptureSettings->bAutoRestartCaptureOnChange);
 				}
 			}
 		}
diff --git a/Engine/Plugins/Media/MediaFrameworkUtilities/Source/MediaFrameworkUtilitiesEditor/Private/AssetEditor/SMediaProfileMediaItemDisplay.cpp b/Engine/Plugins/Media/MediaFrameworkUtilities/Source/MediaFrameworkUtilitiesEditor/Private/AssetEditor/SMediaProfileMediaItemDisplay.cpp
--- a/Engine/Plugins/Media/MediaFrameworkUtilities/Source/MediaFrameworkUtilitiesEditor/Private/AssetEditor/SMediaProfileMediaItemDisplay.cpp
+++ b/Engine/Plugins/Media/MediaFrameworkUtilities/Source/MediaFrameworkUtilitiesEditor/Private/AssetEditor/SMediaProfileMediaItemDisplay.cpp
@@ -110,6 +110,8 @@ void SMediaProfileMediaSourceDisplay::ConfigureMediaImage()
 			.InvertAlphaChannel(this, &SMediaProfileMediaSourceDisplay::GetInvertAlphaChannelMask)
 			.DrawCheckerboard(this, &SMediaProfileMediaSourceDisplay::GetDrawAlphaBlendedCheckerboard)
 			.ImageSize(TextureSize)
+			.Rotation(this, &SMediaProfileMediaSourceDisplay::GetRotationValue, false)
+			.Scale(this, &SMediaProfileMediaSourceDisplay::GetScale)
 		]);
 }
 
@@ -140,6 +142,11 @@ FText SMediaProfileMediaSourceDisplay::GetBaseMediaTypeLabel() const
 	return LOCTEXT("MediaSourceBaseTypeLabel", "Media Source");
 }
 
+bool SMediaProfileMediaSourceDisplay::CanTransformMediaItem() const
+{
+	return GetMediaItem() != nullptr;
+}
+
 UMediaTexture* SMediaProfileMediaSourceDisplay::GetMediaTexture() const
 {
 	UMediaProfile* MediaProfile = GetMediaProfile();
@@ -163,8 +170,13 @@ FVector2D SMediaProfileMediaSourceDisplay::GetSourceImageSize() const
 
 FOptionalSize SMediaProfileMediaSourceDisplay::GetSourceAspectRatio() const
 {
+	const float Rotation = GetRotationValue();
 	const FVector2D ImageSize = GetSourceImageSize();
-	return (ImageSize.X > 0 && ImageSize.Y > 0) ? FOptionalSize(ImageSize.X / ImageSize.Y) : FOptionalSize();
+	const FVector2D RotatedImageSize = FVector2D(
+		ImageSize.X * FMath::Abs(FMath::Cos(Rotation)) + ImageSize.Y * FMath::Abs(FMath::Sin(Rotation)),
+		ImageSize.X * FMath::Abs(FMath::Sin(Rotation)) + ImageSize.Y * FMath::Abs(FMath::Cos(Rotation)));
+	
+	return (RotatedImageSize.X > 0 && RotatedImageSize.Y > 0) ? FOptionalSize(RotatedImageSize.X / RotatedImageSize.Y) : FOptionalSize();
 }
 
 FText SMediaProfileMediaSourceDisplay::GetTimeDurationText() const
@@ -854,6 +866,8 @@ public:
 		SLATE_ATTRIBUTE(EColorChannelMask, ChannelMask)
 		SLATE_ATTRIBUTE(bool, InvertAlphaChannel)
 		SLATE_ATTRIBUTE(bool, DrawCheckerboard)
+		SLATE_ATTRIBUTE(float, Rotation)
+		SLATE_ATTRIBUTE(FVector2f, Scale)
 	SLATE_END_ARGS()
 
 	void Construct(const FArguments& InArgs, UMediaTexture* InMediaTexture)
@@ -861,6 +875,8 @@ public:
 		ChannelMaskAttr = InArgs._ChannelMask;
 		InvertAlphaChannelAttr = InArgs._InvertAlphaChannel;
 		DrawCheckerboardAttr = InArgs._DrawCheckerboard;
+		RotationAttr = InArgs._Rotation;
+		ScaleAttr = InArgs._Scale;
 		
 		SMediaCaptureImageBase::FArguments BaseArgs;
 		BaseArgs.MediaCaptureState(InArgs._MediaCaptureState);
@@ -889,6 +905,8 @@ public:
 				.ChannelMask(ChannelMaskAttr)
 				.InvertAlphaChannel(InvertAlphaChannelAttr)
 				.DrawCheckerboard(DrawCheckerboardAttr)
+				.Rotation(RotationAttr)
+				.Scale(ScaleAttr)
 			);
 		}
 		else
@@ -928,6 +946,12 @@ private:
 
 	/** Attribute for drawing the alpha-blended checkerboard background */
 	TAttribute<bool> DrawCheckerboardAttr;
+	
+	/** Attribute for the rotation to apply to the media texture */
+	TAttribute<float> RotationAttr;
+	
+	/** Attribute for the scale to apply to the media texture */
+	TAttribute<FVector2f> ScaleAttr;
 };
 
 void SMediaProfileMediaOutputDisplay::Construct(const FArguments& InArgs, const TSharedPtr<SMediaProfileViewport>& InOwningViewport)
@@ -1091,6 +1115,8 @@ void SMediaProfileMediaOutputDisplay::ConfigureMediaImage()
 				.ChannelMask(this, &SMediaProfileMediaOutputDisplay::GetChannelMask)
 				.InvertAlphaChannel(this, &SMediaProfileMediaOutputDisplay::GetInvertAlphaChannelMask)
 				.DrawCheckerboard(this, &SMediaProfileMediaOutputDisplay::GetDrawAlphaBlendedCheckerboard)
+				.Rotation(this, &SMediaProfileMediaOutputDisplay::GetRotationValue, false)
+				.Scale(this, &SMediaProfileMediaOutputDisplay::GetScale)
 			);
 			return;
 		}
@@ -1133,6 +1159,218 @@ FText SMediaProfileMediaOutputDisplay::GetBaseMediaTypeLabel() const
 	return LOCTEXT("MediaOutputBaseTypeLabel", "Media Output");
 }
 
+bool SMediaProfileMediaOutputDisplay::CanTransformMediaItem() const
+{
+	if (CaptureType == ECaptureType::MediaTexture)
+	{
+		if (UMediaProfileEditorCaptureSettings* CaptureSettings = FMediaProfileEditor::GetMediaFrameworkCaptureSettings())
+		{
+			if (CaptureSettings->MediaTextureCaptures.IsValidIndex(CaptureIndex) &&
+				CaptureSettings->MediaTextureCaptures[CaptureIndex].MediaOutput == GetMediaItem())
+			{
+				return true;
+			}
+		}
+	}
+	
+	return SMediaProfileMediaItemDisplayBase<UMediaOutput>::CanTransformMediaItem();
+}
+
+EMediaImageRotation SMediaProfileMediaOutputDisplay::GetRotation()
+{
+	if (CaptureType == ECaptureType::MediaTexture)
+	{
+		if (UMediaProfileEditorCaptureSettings* CaptureSettings = FMediaProfileEditor::GetMediaFrameworkCaptureSettings())
+		{
+			if (CaptureSettings->MediaTextureCaptures.IsValidIndex(CaptureIndex) &&
+				CaptureSettings->MediaTextureCaptures[CaptureIndex].MediaOutput == GetMediaItem())
+			{
+				return (EMediaImageRotation)CaptureSettings->MediaTextureCaptures[CaptureIndex].Transform.Rotation;
+			}
+		}
+	}
+	
+	return SMediaProfileMediaItemDisplayBase<UMediaOutput>::GetRotation();
+}
+
+ECheckBoxState SMediaProfileMediaOutputDisplay::GetRotationCheckState(EMediaImageRotation InRotation)
+{
+	if (CaptureType == ECaptureType::MediaTexture)
+	{
+		if (UMediaProfileEditorCaptureSettings* CaptureSettings = FMediaProfileEditor::GetMediaFrameworkCaptureSettings())
+		{
+			if (CaptureSettings->MediaTextureCaptures.IsValidIndex(CaptureIndex) &&
+				CaptureSettings->MediaTextureCaptures[CaptureIndex].MediaOutput == GetMediaItem())
+			{
+				return (EMediaImageRotation)CaptureSettings->MediaTextureCaptures[CaptureIndex].Transform.Rotation == InRotation ? ECheckBoxState::Checked : ECheckBoxState::Unchecked;
+			}
+		}
+	}
+	
+	return SMediaProfileMediaItemDisplayBase<UMediaOutput>::GetRotationCheckState(InRotation);
+}
+
+float SMediaProfileMediaOutputDisplay::GetRotationValue(bool bInDegrees) const
+{
+	if (CaptureType == ECaptureType::MediaTexture)
+	{
+		if (UMediaProfileEditorCaptureSettings* CaptureSettings = FMediaProfileEditor::GetMediaFrameworkCaptureSettings())
+		{
+			if (CaptureSettings->MediaTextureCaptures.IsValidIndex(CaptureIndex) &&
+				CaptureSettings->MediaTextureCaptures[CaptureIndex].MediaOutput == GetMediaItem())
+			{
+				return CaptureSettings->MediaTextureCaptures[CaptureIndex].Transform.GetRotationValue(bInDegrees);
+			}
+		}
+	}
+	
+	return SMediaProfileMediaItemDisplayBase<UMediaOutput>::GetRotationValue(bInDegrees);
+}
+
+void SMediaProfileMediaOutputDisplay::SetRotation(EMediaImageRotation InRotation)
+{
+	if (CaptureType == ECaptureType::MediaTexture)
+	{
+		if (UMediaProfileEditorCaptureSettings* CaptureSettings = FMediaProfileEditor::GetMediaFrameworkCaptureSettings())
+		{
+			if (CaptureSettings->MediaTextureCaptures.IsValidIndex(CaptureIndex) &&
+				CaptureSettings->MediaTextureCaptures[CaptureIndex].MediaOutput == GetMediaItem())
+			{
+				CaptureSettings->MediaTextureCaptures[CaptureIndex].Transform.Rotation = (EMediaCaptureRotation)InRotation;
+				if (CaptureSettings->MediaTextureCaptures[CaptureIndex].Transform.Rotation != EMediaCaptureRotation::None)
+				{
+					CaptureSettings->MediaTextureCaptures[CaptureIndex].CaptureOptions.ResizeMethod = EMediaCaptureResizeMethod::ResizeInRenderPass;
+				}
+				
+				if (TSharedPtr<FMediaProfileEditor> PinnedMediaProfileEditor = MediaProfileEditor.Pin())
+				{
+					PinnedMediaProfileEditor->RestartActiveMediaCaptures(GetMediaItem());
+				}
+				
+				return;
+			}
+		}
+	}
+	
+	SMediaProfileMediaItemDisplayBase<UMediaOutput>::SetRotation(InRotation);
+}
+
+void SMediaProfileMediaOutputDisplay::GetFlipState(bool& bOutFlipHorizontal, bool& bOutFlipVertical)
+{
+	if (CaptureType == ECaptureType::MediaTexture)
+	{
+		if (UMediaProfileEditorCaptureSettings* CaptureSettings = FMediaProfileEditor::GetMediaFrameworkCaptureSettings())
+		{
+			if (CaptureSettings->MediaTextureCaptures.IsValidIndex(CaptureIndex) &&
+				CaptureSettings->MediaTextureCaptures[CaptureIndex].MediaOutput == GetMediaItem())
+			{
+				bOutFlipHorizontal = CaptureSettings->MediaTextureCaptures[CaptureIndex].Transform.bFlipHorizontal;
+				bOutFlipVertical = CaptureSettings->MediaTextureCaptures[CaptureIndex].Transform.bFlipVertical;
+				return;
+			}
+		}
+	}
+	
+	SMediaProfileMediaItemDisplayBase<UMediaOutput>::GetFlipState(bOutFlipHorizontal, bOutFlipVertical);
+}
+
+ECheckBoxState SMediaProfileMediaOutputDisplay::GetFlipHorizontalCheckState()
+{
+	if (CaptureType == ECaptureType::MediaTexture)
+	{
+		if (UMediaProfileEditorCaptureSettings* CaptureSettings = FMediaProfileEditor::GetMediaFrameworkCaptureSettings())
+		{
+			if (CaptureSettings->MediaTextureCaptures.IsValidIndex(CaptureIndex) &&
+				CaptureSettings->MediaTextureCaptures[CaptureIndex].MediaOutput == GetMediaItem())
+			{
+				return CaptureSettings->MediaTextureCaptures[CaptureIndex].Transform.bFlipHorizontal ? ECheckBoxState::Checked : ECheckBoxState::Unchecked;
+			}
+		}
+	}
+	
+	return SMediaProfileMediaItemDisplayBase<UMediaOutput>::GetFlipHorizontalCheckState();
+}
+
+ECheckBoxState SMediaProfileMediaOutputDisplay::GetFlipVerticalCheckState()
+{
+	if (CaptureType == ECaptureType::MediaTexture)
+	{
+		if (UMediaProfileEditorCaptureSettings* CaptureSettings = FMediaProfileEditor::GetMediaFrameworkCaptureSettings())
+		{
+			if (CaptureSettings->MediaTextureCaptures.IsValidIndex(CaptureIndex) &&
+				CaptureSettings->MediaTextureCaptures[CaptureIndex].MediaOutput == GetMediaItem())
+			{
+				return CaptureSettings->MediaTextureCaptures[CaptureIndex].Transform.bFlipVertical ? ECheckBoxState::Checked : ECheckBoxState::Unchecked;
+			}
+		}
+	}
+	
+	return SMediaProfileMediaItemDisplayBase<UMediaOutput>::GetFlipVerticalCheckState();
+}
+
+void SMediaProfileMediaOutputDisplay::ToggleFlipHorizontal()
+{
+	if (CaptureType == ECaptureType::MediaTexture)
+	{
+		if (UMediaProfileEditorCaptureSettings* CaptureSettings = FMediaProfileEditor::GetMediaFrameworkCaptureSettings())
+		{
+			if (CaptureSettings->MediaTextureCaptures.IsValidIndex(CaptureIndex) &&
+				CaptureSettings->MediaTextureCaptures[CaptureIndex].MediaOutput == GetMediaItem())
+			{
+				CaptureSettings->MediaTextureCaptures[CaptureIndex].Transform.bFlipHorizontal = !CaptureSettings->MediaTextureCaptures[CaptureIndex].Transform.bFlipHorizontal;
+				
+				if (TSharedPtr<FMediaProfileEditor> PinnedMediaProfileEditor = MediaProfileEditor.Pin())
+				{
+					PinnedMediaProfileEditor->RestartActiveMediaCaptures(GetMediaItem());
+				}
+				
+				return;
+			}
+		}
+	}
+	
+	SMediaProfileMediaItemDisplayBase<UMediaOutput>::ToggleFlipHorizontal();
+}
+
+void SMediaProfileMediaOutputDisplay::ToggleFlipVertical()
+{
+	if (CaptureType == ECaptureType::MediaTexture)
+	{
+		if (UMediaProfileEditorCaptureSettings* CaptureSettings = FMediaProfileEditor::GetMediaFrameworkCaptureSettings())
+		{
+			if (CaptureSettings->MediaTextureCaptures.IsValidIndex(CaptureIndex) &&
+				CaptureSettings->MediaTextureCaptures[CaptureIndex].MediaOutput == GetMediaItem())
+			{
+				CaptureSettings->MediaTextureCaptures[CaptureIndex].Transform.bFlipVertical = !CaptureSettings->MediaTextureCaptures[CaptureIndex].Transform.bFlipVertical;
+				
+				if (TSharedPtr<FMediaProfileEditor> PinnedMediaProfileEditor = MediaProfileEditor.Pin())
+				{
+					PinnedMediaProfileEditor->RestartActiveMediaCaptures(GetMediaItem());
+				}
+				
+				return;
+			}
+		}
+	}
+	
+	SMediaProfileMediaItemDisplayBase<UMediaOutput>::ToggleFlipVertical();
+}
+
+FVector2f SMediaProfileMediaOutputDisplay::GetScale() const
+{
+	if (UMediaProfileEditorCaptureSettings* CaptureSettings = FMediaProfileEditor::GetMediaFrameworkCaptureSettings())
+	{
+		if (CaptureType == ECaptureType::MediaTexture &&
+			CaptureSettings->MediaTextureCaptures.IsValidIndex(CaptureIndex) &&
+			CaptureSettings->MediaTextureCaptures[CaptureIndex].MediaOutput == GetMediaItem())
+		{
+			return CaptureSettings->MediaTextureCaptures[CaptureIndex].Transform.GetScaleValue();
+		}
+	}
+	
+	return SMediaProfileMediaItemDisplayBase<UMediaOutput>::GetScale();
+}
+
 void SMediaProfileMediaOutputDisplay::SetChannelMask(EColorChannelMask InChannelMask)
 {
 	SMediaProfileMediaItemDisplayBase::SetChannelMask(InChannelMask);
@@ -1318,6 +1556,31 @@ FOptionalSize SMediaProfileMediaOutputDisplay::GetMediaOutputDesiredAspectRatio(
 		TargetSize = GetMutableDefault<UMediaFrameworkMediaCaptureSettings>()->DefaultViewportCaptureSize;
 	}
 	
+	if (CaptureType == ECaptureType::MediaTexture)
+	{
+		if (UMediaProfileEditorCaptureSettings* CaptureSettings = FMediaProfileEditor::GetMediaFrameworkCaptureSettings())
+		{
+			if (CaptureSettings->MediaTextureCaptures.IsValidIndex(CaptureIndex) &&
+				CaptureSettings->MediaTextureCaptures[CaptureIndex].MediaOutput == GetMediaItem())
+			{
+				if (CaptureSettings->MediaTextureCaptures[CaptureIndex].Transform.Rotation != EMediaCaptureRotation::None)
+				{
+					FIntPoint MediaTextureSize = FIntPoint(
+						CaptureSettings->MediaTextureCaptures[CaptureIndex].MediaTexture->GetWidth(),
+						CaptureSettings->MediaTextureCaptures[CaptureIndex].MediaTexture->GetHeight());
+
+					float Rotation = CaptureSettings->MediaTextureCaptures[CaptureIndex].Transform.GetRotationValue();
+
+					const FVector2D RotatedSize = FVector2D(
+						MediaTextureSize.X * FMath::Abs(FMath::Cos(Rotation)) + MediaTextureSize.Y * FMath::Abs(FMath::Sin(Rotation)),
+						MediaTextureSize.X * FMath::Abs(FMath::Sin(Rotation)) + MediaTextureSize.Y * FMath::Abs(FMath::Cos(Rotation)));
+				
+					TargetSize = FIntPoint(FMath::CeilToInt(RotatedSize.X), FMath::CeilToInt(RotatedSize.Y));
+				}
+			}
+		}
+	}
+	
 	const float AspectRatio = (TargetSize.X > 0 && TargetSize.Y > 0) ? (float)TargetSize.X / (float)TargetSize.Y : 1.77777777f;
 	return AspectRatio;
 }
diff --git a/Engine/Plugins/Media/MediaFrameworkUtilities/Source/MediaFrameworkUtilitiesEditor/Private/AssetEditor/SMediaProfileMediaItemDisplay.h b/Engine/Plugins/Media/MediaFrameworkUtilities/Source/MediaFrameworkUtilitiesEditor/Private/AssetEditor/SMediaProfileMediaItemDisplay.h
--- a/Engine/Plugins/Media/MediaFrameworkUtilities/Source/MediaFrameworkUtilitiesEditor/Private/AssetEditor/SMediaProfileMediaItemDisplay.h
+++ b/Engine/Plugins/Media/MediaFrameworkUtilities/Source/MediaFrameworkUtilitiesEditor/Private/AssetEditor/SMediaProfileMediaItemDisplay.h
@@ -428,10 +428,123 @@ protected:
 		LockButton.SetShowInToolbarTopLevel(true);
 		LockButton.ToolBarData.ResizeParams.AllowClipping = false;
 
+		FToolMenuSection& TransformSection = ToolMenu->AddSection(TEXT("TransformSection"), LOCTEXT("TransformSectionLabel", "Transform"));
+		TransformSection.AddEntry(CreateRotationSubmenu());
+		TransformSection.AddEntry(CreateFlipSubmenu());
+		
 		FToolMenuSection& ChannelsSection = ToolMenu->AddSection(TEXT("ChannelsSection"), LOCTEXT("ChannelsSectionLabel", "Channels"));
 		ChannelsSection.AddEntry(CreateChannelsSubmenu(true));
 	}
+	
+	FToolMenuEntry CreateRotationSubmenu()
+	{
+		FToolMenuEntry Entry = FToolMenuEntry::InitSubMenu(
+			"Rotation",
+			TAttribute<FText>::CreateLambda([this]
+			{
+				const EMediaImageRotation Rotation = GetRotation();
+				return FText::Format(LOCTEXT("RotationMenu_LabelFormat", "Rotation ({0})"), StaticEnum<EMediaImageRotation>()->GetDisplayNameTextByValue((int64)Rotation));
+			}),
+			LOCTEXT("RotationMenu_ToolTip", "The rotation of the media item in the viewport"),
+			FNewToolMenuDelegate::CreateLambda([this](UToolMenu* InMenu)
+			{
+				FToolMenuSection& RotationSubMenuSection =
+					InMenu->AddSection("RotationSubMenuSection", LOCTEXT("RotationSubMenuSectionHeader", "Rotation"));
+				
+				auto CreateRotationMenuEntry = [this, &RotationSubMenuSection](EMediaImageRotation InRotation)
+				{
+					RotationSubMenuSection.AddMenuEntry(
+						"",
+						StaticEnum<EMediaImageRotation>()->GetDisplayNameTextByValue((int64)InRotation),
+						FText::GetEmpty(),
+						FSlateIcon(),
+						FUIAction(
+							FExecuteAction::CreateSP(this, &SMediaProfileMediaItemDisplayBase::SetRotation, InRotation),
+							FCanExecuteAction::CreateSP(this, &SMediaProfileMediaItemDisplayBase::CanTransformMediaItem),
+							FGetActionCheckState::CreateSP(this, &SMediaProfileMediaItemDisplayBase::GetRotationCheckState, InRotation)
+						),
+						EUserInterfaceActionType::Check
+					);
+				};
+				
+				CreateRotationMenuEntry(EMediaImageRotation::CCW90);
+				CreateRotationMenuEntry(EMediaImageRotation::None);
+				CreateRotationMenuEntry(EMediaImageRotation::CW90);
+				CreateRotationMenuEntry(EMediaImageRotation::CW180);
+			}),
+			/*bInOpenSubMenuOnClick=*/false,
+			FSlateIcon(FAppStyle::GetAppStyleSetName(), "EditorViewport.RotateMode")
+		);
 
+		return Entry;
+	}
+	
+	FToolMenuEntry CreateFlipSubmenu()
+	{
+		FToolMenuEntry Entry = FToolMenuEntry::InitSubMenu(
+			"Flip",
+			TAttribute<FText>::CreateLambda([this]
+			{
+				bool bFlipHorizontal, bFlipVertical;
+				GetFlipState(bFlipHorizontal, bFlipVertical);
+				FText FlipText = LOCTEXT("FlipNone", "None");
+				if (bFlipHorizontal)
+				{
+					if (bFlipVertical)
+					{
+						FlipText = LOCTEXT("FlipBoth", "Both");
+					}
+					else
+					{
+						FlipText = LOCTEXT("FlipHorizontal", "Horizontal");
+					}
+				}
+				else if (bFlipVertical)
+				{
+					FlipText = LOCTEXT("FlipVertical", "Vertical");
+				}
+				
+				return FText::Format(LOCTEXT("FlipMenu_LabelFormat", "Flip ({0})"), FlipText);
+			}),
+			LOCTEXT("FlipMenu_ToolTip", "The horizontal and vertical flip of the media item in the viewport"),
+			FNewToolMenuDelegate::CreateLambda([this](UToolMenu* InMenu)
+			{
+				FToolMenuSection& FlipSubMenuSection =
+					InMenu->AddSection("FlipSubMenuSection", LOCTEXT("FlipSubMenuSectionHeader", "Flip"));
+
+				FlipSubMenuSection.AddMenuEntry(
+					"FlipHorizontal",
+					LOCTEXT("FlipHorizontalEntry", "Flip Horizontal"),
+					FText::GetEmpty(),
+					FSlateIcon(FAppStyle::GetAppStyleSetName(), "Icons.FlipHorizontal"),
+					FUIAction(
+						FExecuteAction::CreateSP(this, &SMediaProfileMediaItemDisplayBase::ToggleFlipHorizontal),
+						FCanExecuteAction::CreateSP(this, &SMediaProfileMediaItemDisplayBase::CanTransformMediaItem),
+						FGetActionCheckState::CreateSP(this, &SMediaProfileMediaItemDisplayBase::GetFlipHorizontalCheckState)
+					),
+					EUserInterfaceActionType::Check
+				);
+				
+				FlipSubMenuSection.AddMenuEntry(
+					"FlipVertical",
+					LOCTEXT("FlipHorizontalEntry", "Flip Vertical"),
+					FText::GetEmpty(),
+					FSlateIcon(FAppStyle::GetAppStyleSetName(), "Icons.FlipVertical"),
+					FUIAction(
+						FExecuteAction::CreateSP(this, &SMediaProfileMediaItemDisplayBase::ToggleFlipVertical),
+						FCanExecuteAction::CreateSP(this, &SMediaProfileMediaItemDisplayBase::CanTransformMediaItem),
+						FGetActionCheckState::CreateSP(this, &SMediaProfileMediaItemDisplayBase::GetFlipVerticalCheckState)
+					),
+					EUserInterfaceActionType::Check
+				);
+			}),
+			/*bInOpenSubMenuOnClick=*/false,
+			FSlateIcon(FAppStyle::GetAppStyleSetName(), "Icons.FlipVertical")
+		);
+
+		return Entry;
+	}
+	
 	FToolMenuEntry CreateChannelsSubmenu(bool bUseFullLabel)
 	{
 		static FText ChannelLabels[] =
@@ -791,6 +904,145 @@ protected:
 		ChangePanelContents(InMediaItemClass, InMediaItemIndex, bRefreshDisplay);
 	}
 
+	/** Gets whether 2D visual transformations are allowed on the media item */
+	virtual bool CanTransformMediaItem() const
+	{
+		return false;
+	}
+	
+	/** Gets the media item's rotation enum value */
+	virtual EMediaImageRotation GetRotation()
+	{
+		if (FMediaProfileEditorPerMediaItemSettings* MediaItemSettings = GetMediaItemSettings())
+		{
+			return MediaItemSettings->Rotation;
+		}
+		
+		return EMediaImageRotation::None;
+	}
+	
+	/** Gets the media item's actual rotation value */
+	virtual float GetRotationValue(bool bInDegrees = false) const
+	{
+		float RotationInDegrees = 0.0f;
+		if (FMediaProfileEditorPerMediaItemSettings* MediaItemSettings = GetMediaItemSettings())
+		{
+			switch (MediaItemSettings->Rotation)
+			{
+			case EMediaImageRotation::CW90:
+				RotationInDegrees = -90.0f;
+				break;
+			
+			case EMediaImageRotation::CCW90:
+				RotationInDegrees = 90.0f;
+				break;
+			
+			case EMediaImageRotation::CW180:
+				RotationInDegrees = 180.0f;
+				break;
+			
+			default:
+				RotationInDegrees = 0.0f;
+				break;
+			}
+		}
+	
+		return bInDegrees ? RotationInDegrees : FMath::DegreesToRadians(RotationInDegrees);
+	}
+	
+	/** Gets the check state for the specified rotation enum value for the media item */
+	virtual ECheckBoxState GetRotationCheckState(EMediaImageRotation InRotation)
+	{
+		if (FMediaProfileEditorPerMediaItemSettings* MediaItemSettings = GetMediaItemSettings())
+		{
+			return MediaItemSettings->Rotation == InRotation ? ECheckBoxState::Checked : ECheckBoxState::Unchecked;
+		}
+		
+		return ECheckBoxState::Unchecked;
+	}
+	
+	/** Sets the media item's rotation */
+	virtual void SetRotation(EMediaImageRotation InRotation)
+	{
+		if (FMediaProfileEditorPerMediaItemSettings* MediaItemSettings = GetMediaItemSettings())
+		{
+			MediaItemSettings->Rotation = InRotation;
+		}
+	}
+	
+	/** Gets the flip state of the media item */
+	virtual void GetFlipState(bool& bOutFlipHorizontal, bool& bOutFlipVertical)
+	{
+		bOutFlipHorizontal = false;
+		bOutFlipVertical = false;
+		
+		if (FMediaProfileEditorPerMediaItemSettings* MediaItemSettings = GetMediaItemSettings())
+		{
+			bOutFlipHorizontal = MediaItemSettings->bFlipHorizontal;
+			bOutFlipVertical = MediaItemSettings->bFlipVertical;
+		}
+	}
+	
+	/** Gets the check state for the media item's horizontal flip state */
+	virtual ECheckBoxState GetFlipHorizontalCheckState()
+	{
+		if (FMediaProfileEditorPerMediaItemSettings* MediaItemSettings = GetMediaItemSettings())
+		{
+			return MediaItemSettings->bFlipHorizontal ? ECheckBoxState::Checked : ECheckBoxState::Unchecked;
+		}
+		
+		return ECheckBoxState::Unchecked;
+	}
+	
+	/** Gets the check state for the media item's vertical flip state */
+	virtual ECheckBoxState GetFlipVerticalCheckState()
+	{
+		if (FMediaProfileEditorPerMediaItemSettings* MediaItemSettings = GetMediaItemSettings())
+		{
+			return MediaItemSettings->bFlipVertical ? ECheckBoxState::Checked : ECheckBoxState::Unchecked;
+		}
+		
+		return ECheckBoxState::Unchecked;
+	}
+	
+	/** Toggle's the media item's horizontal flip state */
+	virtual void ToggleFlipHorizontal()
+	{
+		if (FMediaProfileEditorPerMediaItemSettings* MediaItemSettings = GetMediaItemSettings())
+		{
+			MediaItemSettings->bFlipHorizontal = !MediaItemSettings->bFlipHorizontal;
+		}
+	}
+	
+	/** Toggle's the media item's vertical flip state */
+	virtual void ToggleFlipVertical()
+	{
+		if (FMediaProfileEditorPerMediaItemSettings* MediaItemSettings = GetMediaItemSettings())
+		{
+			MediaItemSettings->bFlipVertical = !MediaItemSettings->bFlipVertical;
+		}
+	}
+	
+	/** Gets the scale transform that should be applied to the media item's visualization */
+	virtual FVector2f GetScale() const
+	{
+		FVector2f Scale = FVector2f::UnitVector;
+		if (FMediaProfileEditorPerMediaItemSettings* MediaItemSettings = GetMediaItemSettings())
+		{
+			if (MediaItemSettings->bFlipHorizontal)
+			{
+				Scale.X *= -1.0f;
+			}
+			
+			if (MediaItemSettings->bFlipVertical)
+			{
+				Scale.Y *= -1.0f;
+			}
+		}
+		
+		return Scale;
+	}
+	
 	/** Gets the current color channel being displayed in the viewport */
 	virtual EColorChannelMask GetChannelMask() const
 	{
@@ -1015,6 +1267,7 @@ protected:
 	virtual UMediaSource* GetMediaItem() const override;
 	virtual FText GetMediaItemLabel() const override;
 	virtual FText GetBaseMediaTypeLabel() const override;
+	virtual bool CanTransformMediaItem() const override;
 
 private:
 	UMediaTexture* GetMediaTexture() const;
@@ -1056,6 +1309,18 @@ protected:
 	virtual FText GetMediaItemLabel() const override;
 	virtual FText GetBaseMediaTypeLabel() const override;
 
+	virtual bool CanTransformMediaItem() const override;
+	virtual EMediaImageRotation GetRotation() override;
+	virtual ECheckBoxState GetRotationCheckState(EMediaImageRotation InRotation) override;
+	virtual float GetRotationValue(bool bInDegrees = false) const override;
+	virtual void SetRotation(EMediaImageRotation InRotation) override;
+	virtual void GetFlipState(bool& bOutFlipHorizontal, bool& bOutFlipVertical) override;
+	virtual ECheckBoxState GetFlipHorizontalCheckState() override;
+	virtual ECheckBoxState GetFlipVerticalCheckState() override;
+	virtual void ToggleFlipHorizontal() override;
+	virtual void ToggleFlipVertical() override;
+	virtual FVector2f GetScale() const override;
+
 	virtual void SetChannelMask(EColorChannelMask InChannelMask) override;
 	virtual void SetInvertAlphaChannelMask(bool bInInvertAlphaChannelMask) override;
 	virtual bool CanInvertAlphaChannelMask() const override { return true; }
diff --git a/Engine/Plugins/Media/MediaFrameworkUtilities/Source/MediaFrameworkUtilitiesEditor/Private/MediaFrameworkWorldSettingsAssetUserData.cpp b/Engine/Plugins/Media/MediaFrameworkUtilities/Source/MediaFrameworkUtilitiesEditor/Private/MediaFrameworkWorldSettingsAssetUserData.cpp
--- a/Engine/Plugins/Media/MediaFrameworkUtilities/Source/MediaFrameworkUtilitiesEditor/Private/MediaFrameworkWorldSettingsAssetUserData.cpp
+++ b/Engine/Plugins/Media/MediaFrameworkUtilities/Source/MediaFrameworkUtilitiesEditor/Private/MediaFrameworkWorldSettingsAssetUserData.cpp
@@ -55,6 +55,7 @@ FMediaFrameworkCaptureMediaTextureOutputInfo::FMediaFrameworkCaptureMediaTexture
 	, MediaOutput(nullptr)
 {
 	CaptureOptions.bApplyLinearToSRGBConversion = true;
+	CaptureOptions.ResizeMethod = EMediaCaptureResizeMethod::ResizeInRenderPass;
 }
 
 UMediaFrameworkWorldSettingsAssetUserData::UMediaFrameworkWorldSettingsAssetUserData()
diff --git a/Engine/Plugins/Media/MediaFrameworkUtilities/Source/MediaFrameworkUtilitiesEditor/Private/MediaFrameworkWorldSettingsAssetUserData.h b/Engine/Plugins/Media/MediaFrameworkUtilities/Source/MediaFrameworkUtilitiesEditor/Private/MediaFrameworkWorldSettingsAssetUserData.h
--- a/Engine/Plugins/Media/MediaFrameworkUtilities/Source/MediaFrameworkUtilitiesEditor/Private/MediaFrameworkWorldSettingsAssetUserData.h
+++ b/Engine/Plugins/Media/MediaFrameworkUtilities/Source/MediaFrameworkUtilitiesEditor/Private/MediaFrameworkWorldSettingsAssetUserData.h
@@ -124,6 +124,9 @@ struct FMediaFrameworkCaptureMediaTextureOutputInfo
 	UPROPERTY(EditAnywhere, Category="Media Texture Capture")
 	TObjectPtr<UMediaTexture> MediaTexture;
 
+	UPROPERTY(EditAnywhere, Category="Media Texture Capture")
+	FMediaCaptureTransform Transform;
+	
 	UPROPERTY(EditAnywhere, Category="Media Texture Capture")
 	TObjectPtr<UMediaOutput> MediaOutput;
 
diff --git a/Engine/Plugins/Media/MediaIOFramework/Shaders/MediaIO.usf b/Engine/Plugins/Media/MediaIOFramework/Shaders/MediaIO.usf
--- a/Engine/Plugins/Media/MediaIOFramework/Shaders/MediaIO.usf
+++ b/Engine/Plugins/Media/MediaIOFramework/Shaders/MediaIO.usf
@@ -38,3 +38,14 @@ void MediaDownsample(float4 SvPosition : SV_POSITION, out float4 OutColor : SV_T
 	// Fixed rarely occurring yellow color tint of the whole viewport (certain viewport size, need to investigate more)
 	OutColor.rgb = max(float3(0,0,0), OutColor.rgb);
 }
+
+float4 RotationScale;
+
+void MediaTransform(float4 SvPosition : SV_POSITION, out float4 OutColor : SV_Target0)
+{
+	const float2 Position = SvPosition.xy - 0.5 * Output_Extent;
+	const float2 TransformedPosition = float2(dot(Position, RotationScale.xy), dot(Position, RotationScale.zw));
+	const float2 UV = TransformedPosition * Input_ExtentInverse + float2(0.5, 0.5);
+	
+	OutColor = Texture2DSampleLevel(InputTexture, InputSampler, UV, 0);
+}
\ No newline at end of file
diff --git a/Engine/Plugins/Media/MediaIOFramework/Source/MediaIOCore/Private/MediaCapture.cpp b/Engine/Plugins/Media/MediaIOFramework/Source/MediaIOCore/Private/MediaCapture.cpp
--- a/Engine/Plugins/Media/MediaIOFramework/Source/MediaIOCore/Private/MediaCapture.cpp
+++ b/Engine/Plugins/Media/MediaIOFramework/Source/MediaIOCore/Private/MediaCapture.cpp
@@ -152,6 +152,44 @@ FString FMediaCaptureOptions::GetColorConversionSettingsString() const
 	return ColorConversionSettings.ToString();
 }
 
+float FMediaCaptureTransform::GetRotationValue(bool bInDegrees) const
+{
+	float RotationValue;
+	switch (Rotation)
+	{
+	case EMediaCaptureRotation::CW90:
+		RotationValue = -90.0f;
+		break;
+	case EMediaCaptureRotation::CCW90:
+		RotationValue = 90.0f;
+		break;
+	case EMediaCaptureRotation::CW180:
+		RotationValue = 180.0f;
+		break;
+	case EMediaCaptureRotation::None:
+	default: 
+		RotationValue = 0.0f;
+		break;
+	}
+	
+	return bInDegrees ? RotationValue : FMath::DegreesToRadians(RotationValue);
+}
+
+FVector2f FMediaCaptureTransform::GetScaleValue() const
+{
+	FVector2f Scale = FVector2f::UnitVector;
+	if (bFlipHorizontal)
+	{
+		Scale.X *= -1.0f;
+	}
+	if (bFlipVertical)
+	{
+		Scale.Y *= -1.0f;
+	}
+	
+	return Scale;
+}
+
 /* UMediaCapture
 *****************************************************************************/
 
@@ -264,10 +302,10 @@ bool UMediaCapture::CaptureTextureRenderTarget2D(UTextureRenderTarget2D* RenderT
 	return StartSourceCapture(MakeShared<FRenderTargetCaptureSource>(this, CaptureOptions, RenderTarget));
 }
 
-bool UMediaCapture::CaptureMediaTexture(UMediaTexture* InMediaTexture, FMediaCaptureOptions CaptureOptions)
+bool UMediaCapture::CaptureMediaTexture(UMediaTexture* InMediaTexture, FMediaCaptureOptions CaptureOptions, float Rotation, const FVector2f& Scale)
 {
 	using namespace UE::MediaCapture::Private;
-	return StartSourceCapture(MakeShared<FMediaTextureCaptureSource>(this, CaptureOptions, InMediaTexture));
+	return StartSourceCapture(MakeShared<FMediaTextureCaptureSource>(this, CaptureOptions, InMediaTexture, Rotation, Scale));
 }
 
 bool UMediaCapture::StartSourceCapture(TSharedPtr<UE::MediaCapture::Private::FCaptureSource> InSource)
@@ -335,6 +373,8 @@ bool UMediaCapture::StartSourceCapture(TSharedPtr<UE::MediaCapture::Private::FCa
 		//no lock required the command on the render thread is not active yet
 		CaptureSource = MoveTemp(InSource);
 
+		CaptureSource->AddPasses(*CaptureRenderPipeline);
+		
 		if (DesiredCaptureOptions.CapturePhase != EMediaCapturePhase::EndFrame)
 		{
 			/** If we are running with color conversion, we need our extension to be executed before the OCIO extension is applied to the viewport. */
diff --git a/Engine/Plugins/Media/MediaIOFramework/Source/MediaIOCore/Private/MediaCaptureRenderPass.cpp b/Engine/Plugins/Media/MediaIOFramework/Source/MediaIOCore/Private/MediaCaptureRenderPass.cpp
--- a/Engine/Plugins/Media/MediaIOFramework/Source/MediaIOCore/Private/MediaCaptureRenderPass.cpp
+++ b/Engine/Plugins/Media/MediaIOFramework/Source/MediaIOCore/Private/MediaCaptureRenderPass.cpp
@@ -570,6 +570,16 @@ namespace UE::MediaCapture
 		RenderPasses = FRenderPass::CreateRenderPasses(InMediaCapture);
 	}
 
+	void FRenderPipeline::AddPass(const FRenderPass& InPass, const FName& AfterPass)
+	{
+		int32 Index = RenderPasses.IndexOfByPredicate([&AfterPass](const FRenderPass& RenderPass)
+		{
+			return RenderPass.Name == AfterPass;
+		});
+		
+		RenderPasses.Insert(InPass, Index + 1);
+	}
+
 	void FRenderPipeline::InitializeResources_RenderThread(const TSharedPtr<MediaCaptureData::FCaptureFrame>& CapturingFrame, FRDGBuilder& GraphBuilder, FRDGTextureRef InputResource) const
 	{
 		if (CapturingFrame->bMediaCaptureActive)
diff --git a/Engine/Plugins/Media/MediaIOFramework/Source/MediaIOCore/Private/MediaCaptureRenderPass.h b/Engine/Plugins/Media/MediaIOFramework/Source/MediaIOCore/Private/MediaCaptureRenderPass.h
--- a/Engine/Plugins/Media/MediaIOFramework/Source/MediaIOCore/Private/MediaCaptureRenderPass.h
+++ b/Engine/Plugins/Media/MediaIOFramework/Source/MediaIOCore/Private/MediaCaptureRenderPass.h
@@ -173,6 +173,9 @@ namespace UE::MediaCapture
 	public:
 		FRenderPipeline(UMediaCapture* InMediaCapture);
 
+		/** Adds a render pass to the media capture render pipeline, optionally positioning it after the pass with the specified name */
+		void AddPass(const FRenderPass& InPass, const FName& AfterPass = NAME_None);
+		
 		/** Initialize and register render pass resources for a given capture frame. */
 		void InitializeResources_RenderThread(const TSharedPtr<MediaCaptureData::FCaptureFrame>& CapturingFrame, FRDGBuilder& GraphBuilder, FRDGTextureRef InputResource) const;
 
diff --git a/Engine/Plugins/Media/MediaIOFramework/Source/MediaIOCore/Private/MediaCaptureSources.inl b/Engine/Plugins/Media/MediaIOFramework/Source/MediaIOCore/Private/MediaCaptureSources.inl
--- a/Engine/Plugins/Media/MediaIOFramework/Source/MediaIOCore/Private/MediaCaptureSources.inl
+++ b/Engine/Plugins/Media/MediaIOFramework/Source/MediaIOCore/Private/MediaCaptureSources.inl
@@ -10,6 +10,7 @@
 #include "MediaIOCoreModule.h"
 #include "MediaOutput.h"
 #include "MediaTexture.h"
+#include "PixelShaderUtils.h"
 #include "Slate/SceneViewport.h"
 #include "TextureResource.h"
 #include "UObject/WeakObjectPtr.h"
@@ -133,7 +134,7 @@ namespace UE::MediaCapture::Private
 		return true;
 	}
 
-	bool ValidateMediaTexture(const UMediaTexture* InMediaTexture, const FMediaCaptureOptions& CaptureOptions, const FIntPoint& DesiredSize, const EPixelFormat DesiredPixelFormat, const bool bCurrentlyCapturing)
+	bool ValidateMediaTexture(const UMediaTexture* InMediaTexture, const FMediaCaptureOptions& CaptureOptions, const FIntPoint& TransformedMediaTextureSize, const FIntPoint& DesiredSize, const EPixelFormat DesiredPixelFormat, const bool bCurrentlyCapturing)
 	{
 		if (InMediaTexture == nullptr)
 		{
@@ -142,7 +143,7 @@ namespace UE::MediaCapture::Private
 			return false;
 		}
 
-		if (CaptureOptions.ResizeMethod != EMediaCaptureResizeMethod::ResizeInRenderPass && !ValidateSize(FIntPoint(InMediaTexture->GetWidth(), InMediaTexture->GetHeight()), DesiredSize, CaptureOptions, bCurrentlyCapturing))
+		if (CaptureOptions.ResizeMethod != EMediaCaptureResizeMethod::ResizeInRenderPass && !ValidateSize(TransformedMediaTextureSize, DesiredSize, CaptureOptions, bCurrentlyCapturing))
 		{
 			return false;
 		}
@@ -202,6 +203,9 @@ namespace UE::MediaCapture::Private
 			return nullptr;
 		}
 
+		/** Allows the capture source to add its own render passes to the media capture render pipeline */
+		virtual void AddPasses(FRenderPipeline& Pipeline) { }
+
 	public:
 		TWeakObjectPtr<UMediaCapture> MediaCapture;
 		FMediaCaptureOptions CaptureOptions;
@@ -441,18 +445,68 @@ namespace UE::MediaCapture::Private
 		TWeakObjectPtr<UTextureRenderTarget2D> RenderTarget;
 	};
 
+	namespace MediaTextureCapture
+	{
+		class FMediaTextureTransformPS : public FGlobalShader
+		{
+		public:
+			DECLARE_GLOBAL_SHADER(FMediaTextureTransformPS);
+			SHADER_USE_PARAMETER_STRUCT(FMediaTextureTransformPS, FGlobalShader);
+
+			BEGIN_SHADER_PARAMETER_STRUCT(FParameters, )
+				SHADER_PARAMETER_STRUCT(FScreenPassTextureViewportParameters, Input)
+				SHADER_PARAMETER_STRUCT(FScreenPassTextureViewportParameters, Output)
+				SHADER_PARAMETER_RDG_TEXTURE(Texture2D, InputTexture) 
+				SHADER_PARAMETER_SAMPLER(SamplerState, InputSampler)
+				SHADER_PARAMETER(FVector4f, RotationScale)
+				RENDER_TARGET_BINDING_SLOTS()
+			END_SHADER_PARAMETER_STRUCT()
+
+			static bool ShouldCompilePermutation(const FGlobalShaderPermutationParameters& Parameters)
+			{
+				return true;
+			}
+
+			/** Allocates and setup shader parameter in the incoming graph builder */
+			static FParameters* AllocateAndSetParameters(FRDGBuilder& GraphBuilder, const FScreenPassTexture& InputTexture, const FScreenPassTexture& OutputTexture, const FMatrix2x2f& RotationScale)
+			{
+				FParameters* Parameters = GraphBuilder.AllocParameters<FParameters>();
+
+				FVector4f RotScaleVector;
+				RotationScale.GetMatrix(RotScaleVector.X, RotScaleVector.Y, RotScaleVector.Z, RotScaleVector.W);
+				
+				Parameters->Input = GetScreenPassTextureViewportParameters(FScreenPassTextureViewport(InputTexture));
+				Parameters->Output = GetScreenPassTextureViewportParameters(FScreenPassTextureViewport(OutputTexture));
+				Parameters->InputTexture = InputTexture.Texture;
+				Parameters->InputSampler = TStaticSamplerState<SF_Bilinear, AM_Border, AM_Border, AM_Border>::GetRHI();
+				Parameters->RotationScale = RotScaleVector;
+				Parameters->RenderTargets[0] = FRenderTargetBinding { OutputTexture.Texture, ERenderTargetLoadAction::ENoAction };
+
+				return Parameters;
+			}
+		};
+	
+		IMPLEMENT_GLOBAL_SHADER(FMediaTextureTransformPS, "/MediaIOShaders/MediaIO.usf", "MediaTransform", SF_Pixel);
+	}
+	
 	class FMediaTextureCaptureSource : public FCaptureSource
 	{
 	public:
-		FMediaTextureCaptureSource(UMediaCapture* InMediaCapture, const FMediaCaptureOptions& InCaptureOptions, UMediaTexture* InMediaTexture)
+		FMediaTextureCaptureSource(UMediaCapture* InMediaCapture, const FMediaCaptureOptions& InCaptureOptions, UMediaTexture* InMediaTexture, float InRotation, const FVector2f& InScale)
 			: FCaptureSource(InMediaCapture, InCaptureOptions)
 			, MediaTexture(InMediaTexture)
 		{
+			float SinAngle, CosAngle;
+			FMath::SinCos(&SinAngle, &CosAngle, FMath::DegreesToRadians(InRotation));
+
+			RotationScale = FMatrix2x2f(
+				InScale.X * CosAngle, -InScale.X * SinAngle,
+				InScale.Y * SinAngle, InScale.Y * CosAngle);
 		}
 
 		virtual bool ValidateSource(const FIntPoint& DesiredSize, const EPixelFormat DesiredPixelFormat, const bool bCurrentlyCapturing) override
 		{
-			return ValidateMediaTexture(MediaTexture.Get(), CaptureOptions, DesiredSize, DesiredPixelFormat, bCurrentlyCapturing);
+			return ValidateMediaTexture(MediaTexture.Get(), CaptureOptions, GetSize(), DesiredSize, DesiredPixelFormat, bCurrentlyCapturing);
 		}
 		
 		virtual bool PostInitialize() override
@@ -473,7 +527,9 @@ namespace UE::MediaCapture::Private
 		{
 			if (TStrongObjectPtr<UMediaTexture> PinnedMediaTexture = MediaTexture.Pin())
 			{
-				return FIntPoint(PinnedMediaTexture->GetWidth(), PinnedMediaTexture->GetHeight());
+				const FVector2f InitialSize = FVector2f(PinnedMediaTexture->GetWidth(), PinnedMediaTexture->GetHeight());
+				const FVector2f TransformedSize = RotationScale.TransformVector(InitialSize).GetAbs();
+				return FIntPoint(FMath::CeilToInt(TransformedSize.X), FMath::CeilToInt(TransformedSize.Y));
 			}
 
 			return FIntPoint();
@@ -505,8 +561,73 @@ namespace UE::MediaCapture::Private
 			return nullptr;
 		}
 		
+		virtual void AddPasses(FRenderPipeline& Pipeline) override
+		{
+			if (!RotationScale.IsNearlyIdentity())
+			{
+				FRenderPass TransformRenderPass;
+				TransformRenderPass.Name = "TransformMediaTexture";
+				TransformRenderPass.OutputType = ERDGViewableResourceType::Texture;
+				TransformRenderPass.InitializePassOutputDelegate = FRenderPass::FInitializePassOutput::CreateLambda([](const UE::MediaCapture::FInitializePassOutputArgs& Args, uint32 FrameId)
+				{
+					check(Args.InputResource && Args.InputResource->Type == ERDGViewableResourceType::Texture);
+					FRDGTextureDesc InputResourceDesc = static_cast<FRDGTextureRef>(Args.InputResource)->Desc;
+			
+					FRDGTextureDesc TransformOutputTextureDesc = InputResourceDesc;
+					TransformOutputTextureDesc.Reset();
+					TransformOutputTextureDesc.Flags = TexCreate_None;
+					TransformOutputTextureDesc.Flags |= TexCreate_RenderTargetable | TexCreate_ShaderResource | TexCreate_UAV;
+					TransformOutputTextureDesc.ClearValue = FClearValueBinding(FLinearColor(0, 0, 0, 0));
+					TransformOutputTextureDesc.Extent = Args.MediaCapture->GetDesiredSize();
+
+					TRefCountPtr<IPooledRenderTarget> RenderTarget = AllocatePooledTexture(TransformOutputTextureDesc, TEXT("MediaCapture Transformed Output"));
+					return RenderTargetResource(MoveTemp(RenderTarget));
+				});
+				TransformRenderPass.ExecutePassDelegate = FRenderPass::FExecutePass::CreateLambda([RotScale = RotationScale, Size = GetSize()](const UE::MediaCaptureData::FCaptureFrameArgs& Args, const TSharedPtr<UE::MediaCaptureData::FCaptureFrame>& CapturingFrame, FRDGViewableResource* InputResource, FRDGViewableResource* OutputResource)
+				{
+					check(InputResource->Type == ERDGViewableResourceType::Texture);
+
+					const float FitScale = FMath::Min(FMath::Min(Args.DesiredSize.X / (float)Size.X, Args.DesiredSize.Y / (float)Size.Y), 1.0f);
+					const FMatrix2x2f FitRotScale = RotScale.Concatenate(FMatrix2x2f(1.0f / FitScale));
+					const FIntPoint FitSize = FIntPoint(FMath::CeilToInt(Size.X * FitScale), FMath::CeilToInt(Size.Y * FitScale));
+					
+					FScreenPassTexture InputTexture(static_cast<FRDGTextureRef>(InputResource));
+					FScreenPassRenderTarget Input;
+					const FIntVector InputSize = InputTexture.Texture->Desc.GetSize();
+					Input.Texture = InputTexture.Texture;
+					Input.ViewRect = FIntRect{ 0, 0, InputSize.X, InputSize.Y };
+					Input.LoadAction = ERenderTargetLoadAction::ENoAction;
+					Input.UpdateVisualizeTextureExtent();
+
+					FScreenPassRenderTarget OutputTexture(static_cast<FRDGTextureRef>(OutputResource), ERenderTargetLoadAction::ENoAction);
+					const FIntVector OutputSize = OutputTexture.Texture->Desc.GetSize();
+					OutputTexture.ViewRect = FIntRect {
+						FMath::CeilToInt(0.5f * (OutputSize.X - FitSize.X)), FMath::CeilToInt(0.5f * (OutputSize.Y - FitSize.Y)),
+						FMath::FloorToInt(0.5f * (OutputSize.X + FitSize.X)), FMath::FloorToInt(0.5f * (OutputSize.Y + FitSize.Y))};
+					OutputTexture.UpdateVisualizeTextureExtent();
+
+					MediaTextureCapture::FMediaTextureTransformPS::FParameters* PassParameters =
+						MediaTextureCapture::FMediaTextureTransformPS::AllocateAndSetParameters(Args.GraphBuilder, InputTexture, OutputTexture, FitRotScale);
+					
+					const FGlobalShaderMap* GlobalShaderMap = GetGlobalShaderMap(GMaxRHIFeatureLevel);
+					const TShaderMapRef<MediaTextureCapture::FMediaTextureTransformPS> PixelShader(GlobalShaderMap);
+					FPixelShaderUtils::AddFullscreenPass(
+						Args.GraphBuilder,
+						GlobalShaderMap,
+						RDG_EVENT_NAME("MediaCapture Texture Transform"),
+						PixelShader,
+						PassParameters,
+						OutputTexture.ViewRect);
+					
+				});
+			
+				Pipeline.AddPass(TransformRenderPass, TEXT("Uncompression"));
+			}
+		}
+		
 	private:
 		TWeakObjectPtr<UMediaTexture> MediaTexture;
+		FMatrix2x2f RotationScale;
 	};
 }
 
diff --git a/Engine/Plugins/Media/MediaIOFramework/Source/MediaIOCore/Public/MediaCapture.h b/Engine/Plugins/Media/MediaIOFramework/Source/MediaIOCore/Public/MediaCapture.h
--- a/Engine/Plugins/Media/MediaIOFramework/Source/MediaIOCore/Public/MediaCapture.h
+++ b/Engine/Plugins/Media/MediaIOFramework/Source/MediaIOCore/Public/MediaCapture.h
@@ -254,6 +254,33 @@ private:
 #endif
 };
 
+UENUM()
+enum class EMediaCaptureRotation : uint8
+{
+	None UMETA(DisplayName = "0°"),
+	CW90 UMETA(DisplayName = "90°"),
+	CCW90 UMETA(DisplayName = "-90°"),
+	CW180 UMETA(DisplayName = "180°")
+};
+
+USTRUCT()
+struct FMediaCaptureTransform
+{
+	GENERATED_BODY()
+	
+	UPROPERTY(EditAnywhere, Category="Media Texture Capture")
+	EMediaCaptureRotation Rotation = EMediaCaptureRotation::None;
+	
+	UPROPERTY(EditAnywhere, Category="Media Texture Capture")
+	bool bFlipHorizontal = false;
+	
+	UPROPERTY(EditAnywhere, Category="Media Texture Capture")
+	bool bFlipVertical = false;
+	
+	UE_API float GetRotationValue(bool bInDegrees = false) const;
+	UE_API FVector2f GetScaleValue() const;
+};
+
 
 /** Delegate signatures */
 DECLARE_DYNAMIC_MULTICAST_DELEGATE(FMediaCaptureStateChangedSignature);
@@ -322,7 +349,7 @@ public:
 	 * @return True if the capture was successfully started
 	 */
 	UFUNCTION(BlueprintCallable, Category = "Media|Output")
-	UE_API bool CaptureMediaTexture(UMediaTexture* InMediaTexture, FMediaCaptureOptions CaptureOptions);
+	UE_API bool CaptureMediaTexture(UMediaTexture* InMediaTexture, FMediaCaptureOptions CaptureOptions, float Rotation, const FVector2f& Scale);
 	
 	/**
 	 * Update the current capture with a SceneViewport.
