Media Profile: Fix localization name collision

#Synced-CL 52191480

diff --git a/Engine/Plugins/Media/MediaFrameworkUtilities/Source/MediaFrameworkUtilitiesEditor/Private/AssetEditor/SMediaProfileMediaItemDisplay.h b/Engine/Plugins/Media/MediaFrameworkUtilities/Source/MediaFrameworkUtilitiesEditor/Private/AssetEditor/SMediaProfileMediaItemDisplay.h
--- a/Engine/Plugins/Media/MediaFrameworkUtilities/Source/MediaFrameworkUtilitiesEditor/Private/AssetEditor/SMediaProfileMediaItemDisplay.h
+++ b/Engine/Plugins/Media/MediaFrameworkUtilities/Source/MediaFrameworkUtilitiesEditor/Private/AssetEditor/SMediaProfileMediaItemDisplay.h
@@ -527,7 +527,7 @@ protected:
 				
 				FlipSubMenuSection.AddMenuEntry(
 					"FlipVertical",
-					LOCTEXT("FlipHorizontalEntry", "Flip Vertical"),
+					LOCTEXT("FlipVerticalEntry", "Flip Vertical"),
 					FText::GetEmpty(),
 					FSlateIcon(FAppStyle::GetAppStyleSetName(), "Icons.FlipVertical"),
 					FUIAction(
