Media Profile: Added ability to save and load editor layouts that store the tab configuration and the viewport configuration of the media profile editor

#Synced-CL 52369692

diff --git a/Engine/Plugins/Media/MediaFrameworkUtilities/Source/MediaFrameworkUtilities/Private/Profile/MediaProfile.cpp b/Engine/Plugins/Media/MediaFrameworkUtilities/Source/MediaFrameworkUtilities/Private/Profile/MediaProfile.cpp
--- a/Engine/Plugins/Media/MediaFrameworkUtilities/Source/MediaFrameworkUtilities/Private/Profile/MediaProfile.cpp
+++ b/Engine/Plugins/Media/MediaFrameworkUtilities/Source/MediaFrameworkUtilities/Private/Profile/MediaProfile.cpp
@@ -89,6 +89,20 @@ int32 UMediaProfile::FindMediaSourceIndex(UMediaSource* MediaSource) const
 	return MediaSources.Find(MediaSource);
 }
 
+#if WITH_EDITOR
+int32 UMediaProfile::FindMediaSourceIndexByLabel(const FString& InLabel) const
+{
+	int32 Index = MediaSourceLabels.Find(InLabel);
+	
+	if (MediaSources.IsValidIndex(Index))
+	{
+		return Index;
+	}
+	
+	return INDEX_NONE;
+}
+#endif
+
 void UMediaProfile::AddMediaSource(UMediaSource* MediaSource)
 {
 	MediaSources.Add(MediaSource);
@@ -256,6 +270,20 @@ int32 UMediaProfile::FindMediaOutputIndex(UMediaOutput* MediaOutput) const
 	return MediaOutputs.Find(MediaOutput);
 }
 
+#if WITH_EDITOR
+int32 UMediaProfile::FindMediaOutputIndexByLabel(const FString& InLabel) const
+{
+	int32 Index = MediaOutputLabels.Find(InLabel);
+	
+	if (MediaOutputs.IsValidIndex(Index))
+	{
+		return Index;
+	}
+	
+	return INDEX_NONE;
+}
+#endif
+
 void UMediaProfile::AddMediaOutput(UMediaOutput* MediaOutput)
 {
 	MediaOutputs.Add(MediaOutput);
diff --git a/Engine/Plugins/Media/MediaFrameworkUtilities/Source/MediaFrameworkUtilities/Public/Profile/MediaProfile.h b/Engine/Plugins/Media/MediaFrameworkUtilities/Source/MediaFrameworkUtilities/Public/Profile/MediaProfile.h
--- a/Engine/Plugins/Media/MediaFrameworkUtilities/Source/MediaFrameworkUtilities/Public/Profile/MediaProfile.h
+++ b/Engine/Plugins/Media/MediaFrameworkUtilities/Source/MediaFrameworkUtilities/Public/Profile/MediaProfile.h
@@ -98,6 +98,15 @@ public:
 	 */
 	int32 FindMediaSourceIndex(UMediaSource* MediaSource) const;
 	
+#if WITH_EDITOR
+	/**
+	 * Gets the index of a media source whose label matches the specified label if it was found in the media profile's list of sources
+	 * @param InLabel The label of the media source to find
+	 * @return The index of the source in the profile's sources list, or INDEX_NONE if it was not found
+	 */
+	int32 FindMediaSourceIndexByLabel(const FString& InLabel) const;
+#endif
+
 	/**
 	 * Adds the specified media source to the profile's list of media sources
 	 * @param MediaSource The media source to add
@@ -161,6 +170,15 @@ public:
 	 */
 	int32 FindMediaOutputIndex(UMediaOutput* MediaOutput) const;
 	
+#if WITH_EDITOR
+	/**
+	 * Gets the index of a media output with the specified label if it was found in the media profile's list of outputs
+	 * @param InLabel The label of the media output to find
+	 * @return The index of the output in the profile's output list, or INDEX_NONE if it was not found
+	 */
+	int32 FindMediaOutputIndexByLabel(const FString& InLabel) const;
+#endif
+
 	/**
 	 * Adds the specified media output to the profile's list of media outputs
 	 * @param MediaOutput The media output to add
diff --git a/Engine/Plugins/Media/MediaFrameworkUtilities/Source/MediaFrameworkUtilitiesEditor/MediaFrameworkUtilitiesEditor.Build.cs b/Engine/Plugins/Media/MediaFrameworkUtilities/Source/MediaFrameworkUtilitiesEditor/MediaFrameworkUtilitiesEditor.Build.cs
--- a/Engine/Plugins/Media/MediaFrameworkUtilities/Source/MediaFrameworkUtilitiesEditor/MediaFrameworkUtilitiesEditor.Build.cs
+++ b/Engine/Plugins/Media/MediaFrameworkUtilities/Source/MediaFrameworkUtilitiesEditor/MediaFrameworkUtilitiesEditor.Build.cs
@@ -28,6 +28,8 @@ namespace UnrealBuildTool.Rules
 					"Engine",
 					"RenderCore",
 					"InputCore",
+					"Json",
+					"JsonUtilities",
 					"LevelEditor",
 					"MainFrame",
 					"MaterialEditor",
diff --git a/Engine/Plugins/Media/MediaFrameworkUtilities/Source/MediaFrameworkUtilitiesEditor/Private/AssetEditor/MediaProfileCommands.cpp b/Engine/Plugins/Media/MediaFrameworkUtilities/Source/MediaFrameworkUtilitiesEditor/Private/AssetEditor/MediaProfileCommands.cpp
--- a/Engine/Plugins/Media/MediaFrameworkUtilities/Source/MediaFrameworkUtilitiesEditor/Private/AssetEditor/MediaProfileCommands.cpp
+++ b/Engine/Plugins/Media/MediaFrameworkUtilities/Source/MediaFrameworkUtilitiesEditor/Private/AssetEditor/MediaProfileCommands.cpp
@@ -11,6 +11,11 @@ void FMediaProfileCommands::RegisterCommands()
 {
 	UI_COMMAND(Apply, "Apply", "Apply changes to the media profile.", EUserInterfaceActionType::Button, FInputChord());
 	UI_COMMAND(Edit, "Edit", "Edit the current media profile.", EUserInterfaceActionType::Button, FInputChord());
+	UI_COMMAND(SaveLayout, "Save Layout", "Save the current layout", EUserInterfaceActionType::Button, FInputChord());
+	UI_COMMAND(SaveLayoutAs, "Save Layout As", "Saves the current layout as...", EUserInterfaceActionType::Button, FInputChord());
+	UI_COMMAND(RemoveAllLayouts, "Remove All Layouts", "Removes all the user created layouts", EUserInterfaceActionType::Button, FInputChord());
+	UI_COMMAND(ToggleFullscreen, "Enable Fullscreen", "Enables fullscreen mode for this window", EUserInterfaceActionType::ToggleButton, FInputChord(EModifierKey::Shift, EKeys::F11));
+	UI_COMMAND(Immersive, "Immersive View", "Toggle immersive view on the viewport", EUserInterfaceActionType::Button, FInputChord(EKeys::F11));
 }
 
 #undef LOCTEXT_NAMESPACE
diff --git a/Engine/Plugins/Media/MediaFrameworkUtilities/Source/MediaFrameworkUtilitiesEditor/Private/AssetEditor/MediaProfileCommands.h b/Engine/Plugins/Media/MediaFrameworkUtilities/Source/MediaFrameworkUtilitiesEditor/Private/AssetEditor/MediaProfileCommands.h
--- a/Engine/Plugins/Media/MediaFrameworkUtilities/Source/MediaFrameworkUtilitiesEditor/Private/AssetEditor/MediaProfileCommands.h
+++ b/Engine/Plugins/Media/MediaFrameworkUtilities/Source/MediaFrameworkUtilitiesEditor/Private/AssetEditor/MediaProfileCommands.h
@@ -27,4 +27,19 @@ public:
 
 	/** Edit the current Media Profile. */
 	TSharedPtr<FUICommandInfo> Edit;
+	
+	/** Saves the current media profile editor layout */
+	TSharedPtr<FUICommandInfo> SaveLayout;
+	
+	/** Saves the current media profile editor layout */
+	TSharedPtr<FUICommandInfo> SaveLayoutAs;
+	
+	/** Removes all media profile editor layouts */
+	TSharedPtr<FUICommandInfo> RemoveAllLayouts;
+	
+	/** Toggles fullscreen mode on the media profile editor window */
+	TSharedPtr<FUICommandInfo> ToggleFullscreen;
+	
+	/** Makes the media profile editor's viewport immersive */
+	TSharedPtr<FUICommandInfo> Immersive;
 };
diff --git a/Engine/Plugins/Media/MediaFrameworkUtilities/Source/MediaFrameworkUtilitiesEditor/Private/AssetEditor/MediaProfileEditor.cpp b/Engine/Plugins/Media/MediaFrameworkUtilities/Source/MediaFrameworkUtilitiesEditor/Private/AssetEditor/MediaProfileEditor.cpp
--- a/Engine/Plugins/Media/MediaFrameworkUtilities/Source/MediaFrameworkUtilitiesEditor/Private/AssetEditor/MediaProfileEditor.cpp
+++ b/Engine/Plugins/Media/MediaFrameworkUtilities/Source/MediaFrameworkUtilitiesEditor/Private/AssetEditor/MediaProfileEditor.cpp
@@ -8,25 +8,31 @@
 #include "Engine/TextureRenderTarget2D.h"
 #include "LevelEditorViewport.h"
 #include "MediaProfileEditorUserSettings.h"
-#include "PropertyEditorModule.h"
 #include "SMediaProfileDetailsPanel.h"
 #include "SMediaProfileSourcesTreeView.h"
 #include "SMediaProfileViewport.h"
 #include "CaptureCardMediaSource.h"
+#include "MediaProfileCommands.h"
+#include "SPrimaryButton.h"
 #include "CaptureTab/SMediaFrameworkCapture.h"
+#include "Dom/JsonObject.h"
 #include "Framework/Application/SlateApplication.h"
 #include "GameFramework/WorldSettings.h"
+#include "Misc/FileHelper.h"
 #include "Modules/ModuleManager.h"
 #include "Profile/IMediaProfileManager.h"
 #include "Profile/MediaProfile.h"
 #include "Profile/MediaProfilePlaybackManager.h"
 #include "Profile/MediaProfileSettings.h"
+#include "Serialization/JsonSerializer.h"
+#include "Serialization/JsonWriter.h"
 #include "Slate/SceneViewport.h"
 #include "UI/MediaFrameworkUtilitiesEditorStyle.h"
 #include "UI/SMediaFrameworkTimecodeGenlockHeader.h"
 #include "UI/SMediaFrameworkTimecodeGenlockPanel.h"
-#include "Widgets/SViewport.h"
 #include "Widgets/Docking/SDockTab.h"
+#include "Widgets/Input/SComboButton.h"
+#include "Widgets/Input/SEditableTextBox.h"
 
 #define LOCTEXT_NAMESPACE "MediaProfileEditor"
 
@@ -35,22 +41,240 @@
 const FName FMediaProfileEditor::MediaTreeTabId =  FName(TEXT("MediaTree"));
 const FName FMediaProfileEditor::DetailsTabId =  FName(TEXT("Details"));
 const FName FMediaProfileEditor::TimecodeTabId =  FName(TEXT("Timecode"));
+const FString FMediaProfileEditor::DefaultLayoutName = TEXT("Default Layout");
 
-namespace MediaProfileEditor
-{	
-	/**
-	 * Determines the target size from the media output if desired size is specified, otherwise,
-	 * fallback to the default viewport capture resolution.
-	 */
-	FIntPoint GetTargetSize(UMediaOutput* InMediaOutput)
+namespace UE::MediaProfileEditor::Private
+{
+	class SSaveLayoutAsDialog : public SCompoundWidget
 	{
-		const FIntPoint TargetSize = InMediaOutput->GetRequestedSize();
-		if (TargetSize != UMediaOutput::RequestCaptureSourceSize)
+	public:
+		SLATE_BEGIN_ARGS(SSaveLayoutAsDialog) {}
+		SLATE_END_ARGS()
+
+		virtual void Construct(const FArguments& InArgs)
 		{
-			return TargetSize;
+			ChildSlot
+			[
+				SNew(SBorder)
+				.BorderImage(FAppStyle::Get().GetBrush("Brushes.Panel"))
+				.Padding(16.0f)
+				[
+					SNew(SBox)
+					.WidthOverride(300.0f)
+					[
+						// Layout Name
+						SNew(SVerticalBox)
+						+ SVerticalBox::Slot()
+						.AutoHeight()
+						[
+							SNew(SHorizontalBox)
+
+							+SHorizontalBox::Slot()
+							.AutoWidth()
+							.VAlign(VAlign_Center)
+							[
+								SNew(STextBlock)
+								.Text(LOCTEXT("SaveLayoutAsLayoutNameLabel", "Name"))
+								.Margin(FMargin(0, 0, 8, 8))
+							]
+							+SHorizontalBox::Slot()
+							.Padding(0, 0, 0, 8)
+							[
+								SNew(SEditableTextBox)
+								.Text_Lambda([this]() { return FText::FromString(EnteredLayoutName); })
+								.OnTextCommitted(this, &SSaveLayoutAsDialog::CommitLayoutName)
+								.OnTextChanged(this, &SSaveLayoutAsDialog::SetLayoutName)
+								.SelectAllTextWhenFocused(true)
+							]
+						]
+						+ SVerticalBox::Slot()
+						.AutoHeight()
+						[
+							// Constant height, whether the label is visible or not
+							SNew(SBox)
+							.HeightOverride(20.f)
+							[
+								SNew(SBorder)
+								.Visibility_Lambda([this] { return bHasError ? EVisibility::Visible : EVisibility::Hidden; })
+								.BorderImage( FAppStyle::GetBrush("AssetDialog.ErrorLabelBorder") )
+								.Content()
+								[
+									SNew(STextBlock)
+									.Text_Lambda([this] { return ErrorText; })
+								]
+							]
+						]
+						+ SVerticalBox::Slot()
+						.HAlign(HAlign_Right)
+						.VAlign(VAlign_Bottom)
+						[
+							SNew(SHorizontalBox)
+							+ SHorizontalBox::Slot()
+							.AutoWidth()
+							.VAlign(VAlign_Bottom)
+							.Padding(0,0,8,0)
+							[
+								SNew(SPrimaryButton)
+								.Text(LOCTEXT("SaveLayoutAsDialogSaveButton", "Save"))
+								.IsEnabled(this, &SSaveLayoutAsDialog::IsConfirmButtonEnabled)
+								.OnClicked(this, &SSaveLayoutAsDialog::OnConfirmButtonClicked)
+							]
+							+ SHorizontalBox::Slot()
+							.AutoWidth()
+							.VAlign(VAlign_Bottom)
+							[
+								SNew(SButton)
+								.Text(LOCTEXT("SaveLayoutAsDialogCancelButton", "Cancel"))
+								.OnClicked(this, &SSaveLayoutAsDialog::OnCancelButtonClicked)
+							]
+						]
+					]
+				]
+			];
 		}
 		
-		return GetMutableDefault<UMediaFrameworkMediaCaptureSettings>()->DefaultViewportCaptureSize;
+		FString GetLayoutName() const
+		{
+			if (bConfirmed)
+			{
+				return EnteredLayoutName;
+			}
+			
+			return TEXT("");
+		}
+		
+	private:
+		void CloseDialog()
+		{
+			TSharedPtr<SWindow> ContainingWindow = FSlateApplication::Get().FindWidgetWindow(AsShared());
+
+			if (ContainingWindow.IsValid())
+			{
+				ContainingWindow->RequestDestroyWindow();
+			}
+		}
+		
+		void SetLayoutName(const FText& InNewName)
+		{
+			bHasError = false;
+			ErrorText = FText::GetEmpty();
+			
+			EnteredLayoutName = InNewName.ToString();
+			
+			bool bIsUniqueName = true;
+			if (const UMediaProfileEditorUserSettings* UserSettings = GetDefault<UMediaProfileEditorUserSettings>())
+			{
+				bIsUniqueName = !UserSettings->SavedLayouts.ContainsByPredicate([Name = EnteredLayoutName](const FMediaProfileEditorLayoutSettings& LayoutSettings)
+				{
+					return LayoutSettings.Name == Name;
+				});
+			}
+			
+			if (!bIsUniqueName)
+			{
+				bHasError = true;
+				ErrorText = LOCTEXT("NonUniqueNameError", "Layout name must be unique");
+			}
+		}
+		
+		void CommitLayoutName(const FText& InNewName, ETextCommit::Type InCommitType)
+		{
+			SetLayoutName(InNewName);
+			
+			if (InCommitType == ETextCommit::OnEnter)
+			{
+				if (!EnteredLayoutName.IsEmpty() && !bHasError)
+				{
+					bConfirmed = true;
+					CloseDialog();
+				}
+			}
+		}
+		
+		bool IsConfirmButtonEnabled() const
+		{
+			return !EnteredLayoutName.IsEmpty() && !bHasError;
+		}
+		
+		FReply OnConfirmButtonClicked()
+		{
+			bConfirmed = true;
+			CloseDialog();
+			return FReply::Handled();
+		}
+		
+		FReply OnCancelButtonClicked()
+		{
+			CloseDialog();
+			return FReply::Handled();
+		}
+		
+	private:
+		FString EnteredLayoutName;
+		bool bConfirmed = false;
+		
+		FText ErrorText = FText::GetEmpty();
+		bool bHasError = false;
+	};
+	
+	FString ShowSaveLayoutAsDialog()
+	{
+		FString NewLayoutName = TEXT("");
+		
+		TSharedRef<SSaveLayoutAsDialog> Dialog = SNew(SSaveLayoutAsDialog);
+		TSharedRef<SWindow> DialogWindow =
+			SNew(SWindow)
+			.Title(LOCTEXT("SaveLayoutAsWindowHeader", "Save Layout As"))
+			.SizingRule(ESizingRule::Autosized);
+
+		DialogWindow->SetContent(Dialog);
+		FSlateApplication::Get().AddModalWindow(DialogWindow, FSlateApplication::Get().GetActiveTopLevelWindow());
+		
+		return Dialog->GetLayoutName();
 	}
+	
+	bool ExportLayoutToString(const TSharedRef<FTabManager::FLayout>& InLayout, const FString& InNewLayoutName, FString& OutLayoutString)
+	{
+		FString SanitizedLayoutName = InNewLayoutName;
+		for (TCHAR& InOutChar : SanitizedLayoutName)
+		{
+			if (!FChar::IsAlnum(InOutChar))
+			{
+				InOutChar = '_';
+			}
+		}
+		
+		const FString LayoutJsonName = "Standalone_MediaProfileEditor_" + SanitizedLayoutName;
+		
+		TSharedRef<FJsonObject> LayoutJson = InLayout->ToJson();
+		LayoutJson->SetStringField(TEXT("Name"), LayoutJsonName);
+		
+		FString LayoutJsonString;
+		TSharedRef<TJsonWriter<>> Writer = TJsonWriterFactory<>::Create(&LayoutJsonString);
+		if (!FJsonSerializer::Serialize(LayoutJson, Writer))
+		{
+			return false;
+		}
+
+		OutLayoutString = LayoutJsonString
+			.Replace(TEXT("{"), TEXT("("))
+			.Replace(TEXT("}"), TEXT(")"))
+			.Replace(TEXT("\r"), TEXT(""))
+			.Replace(TEXT("\n"), TEXT(""))
+			.Replace(TEXT("\t"), TEXT(""));
+		
+		return true;
+	}
+	
+	TSharedPtr<FTabManager::FLayout> ImportLayoutFromString(const FString& InLayoutString)
+	{
+		FString LayoutJson = InLayoutString
+			.Replace(TEXT("("), TEXT("{"))
+			.Replace(TEXT(")"), TEXT("}"))
+			.Replace(TEXT("\\\n"), LINE_TERMINATOR)
+			.Replace(TEXT("\\\r\n"), LINE_TERMINATOR);
+		
+		return FTabManager::FLayout::NewFromString(LayoutJson);
+	}
 }
 
@@ -72,48 +296,15 @@
 	FCoreUObjectDelegates::OnPreObjectPropertyChanged.AddSP(this, &FMediaProfileEditor::OnObjectPreEditChange);
 	FCoreUObjectDelegates::OnObjectPropertyChanged.AddSP(this, &FMediaProfileEditor::OnObjectPropertyChanged);
 	
-	const TSharedRef<FTabManager::FLayout> StandaloneDefaultLayout = FTabManager::NewLayout("Standalone_MediaProfileEditor_Layout")
-		->AddArea
-		(
-			FTabManager::NewPrimaryArea()
-			->SetOrientation(Orient_Vertical)
-			->Split
-			(
-				FTabManager::NewSplitter()
-				->SetOrientation(Orient_Horizontal)
-				->Split
-				(
-					FTabManager::NewStack()
-					->SetSizeCoefficient(0.8f)
-					->SetHideTabWell(true)
-					->AddTab(MediaOutputTabId, ETabState::OpenedTab)
-				)
-				->Split
-				(
-					FTabManager::NewSplitter()
-					->SetOrientation(Orient_Vertical)
-					->SetSizeCoefficient(0.2f)
-					->Split
-					(
-						FTabManager::NewStack()
-						->SetSizeCoefficient(0.4f)
-						->AddTab(MediaTreeTabId, ETabState::OpenedTab)
-					)
-					->Split
-					(
-						FTabManager::NewStack()
-						->SetSizeCoefficient(0.6f)
-						->AddTab(DetailsTabId, ETabState::OpenedTab)
-						->AddTab(TimecodeTabId, ETabState::OpenedTab)
-					)
-				)
-			)
-		);
+	const TSharedRef<FTabManager::FLayout> StandaloneDefaultLayout = GetDefaultLayout();
 
 	const bool bCreateDefaultStandaloneMenu = true;
 	const bool bCreateDefaultToolbar = true;
 	InitAssetEditor(Mode, InitToolkitHost, AppName, StandaloneDefaultLayout, bCreateDefaultStandaloneMenu, bCreateDefaultToolbar, InMediaProfile);
 
+	TabManager->SetOnPersistLayout(FTabManager::FOnPersistLayout::CreateSP(this, &FMediaProfileEditor::OnPersistLayout));
+	
+	BindCommands();
 	ExtendToolbar();
 	RegenerateMenusAndToolbars();
 	
@@ -122,6 +313,14 @@
 	GetMutableDefault<UMediaProfileSettings>()->FillDefaultMediaSourceProxies(InMediaProfile->NumMediaSources(), false);
 	GetMutableDefault<UMediaProfileSettings>()->FillDefaultMediaOutputProxies(InMediaProfile->NumMediaOutputs(), false);
 	IMediaProfileManager::Get().SetCurrentMediaProfile(InMediaProfile);
+	
+	if (const UMediaProfileEditorUserSettings* UserSettings = GetDefault<UMediaProfileEditorUserSettings>())
+	{
+		if (!UserSettings->ActiveLayout.IsEmpty() && UserSettings->ActiveLayout != DefaultLayoutName)
+		{
+			LoadLayout(UserSettings->ActiveLayout);
+		}
+	}
 }
 
 FMediaProfileEditor::~FMediaProfileEditor()
@@ -274,7 +473,7 @@
 		});
 	
 	ForEachOutputInfoForMediaOutput<FMediaFrameworkCaptureRenderTargetCameraOutputInfo>(CaptureSettings, InMediaOutput, 
-		[this, &InMediaOutput, &bHasValidCaptureSettings](const FMediaFrameworkCaptureRenderTargetCameraOutputInfo& OutputInfo)
+		[this, &bHasValidCaptureSettings](const FMediaFrameworkCaptureRenderTargetCameraOutputInfo& OutputInfo)
 		{
 			if (OutputInfo.RenderTarget)
 			{
@@ -299,6 +498,47 @@
 	return GetMutableDefault<UMediaProfileEditorCaptureSettings>();
 }
 
+TSharedRef<FTabManager::FLayout> FMediaProfileEditor::GetDefaultLayout() const
+{
+	return FTabManager::NewLayout("Standalone_MediaProfileEditor_Layout")
+		->AddArea
+		(
+			FTabManager::NewPrimaryArea()
+			->SetOrientation(Orient_Vertical)
+			->Split
+			(
+				FTabManager::NewSplitter()
+				->SetOrientation(Orient_Horizontal)
+				->Split
+				(
+					FTabManager::NewStack()
+					->SetSizeCoefficient(0.8f)
+					->SetHideTabWell(true)
+					->AddTab(MediaOutputTabId, ETabState::OpenedTab)
+				)
+				->Split
+				(
+					FTabManager::NewSplitter()
+					->SetOrientation(Orient_Vertical)
+					->SetSizeCoefficient(0.3f)
+					->Split
+					(
+						FTabManager::NewStack()
+						->SetSizeCoefficient(0.4f)
+						->AddTab(MediaTreeTabId, ETabState::OpenedTab)
+					)
+					->Split
+					(
+						FTabManager::NewStack()
+						->SetSizeCoefficient(0.6f)
+						->AddTab(DetailsTabId, ETabState::OpenedTab)
+						->AddTab(TimecodeTabId, ETabState::OpenedTab)
+					)
+				)
+			)
+		);
+}
+
 TSharedRef<SDockTab> FMediaProfileEditor::SpawnTab_MediaOutput(const FSpawnTabArgs& InArgs)
 {
 	return SNew(SDockTab)
@@ -345,6 +585,33 @@
 	TabManager->TryInvokeTab(InTabId);
 }
 
+void FMediaProfileEditor::BindCommands()
+{
+	ToolkitCommands->MapAction(
+		FMediaProfileCommands::Get().SaveLayout, 
+		FExecuteAction::CreateSP(this, &FMediaProfileEditor::SaveLayout),
+		FCanExecuteAction::CreateSP(this, &FMediaProfileEditor::CanSaveLayout));
+	
+	ToolkitCommands->MapAction(
+		FMediaProfileCommands::Get().SaveLayoutAs, 
+		FExecuteAction::CreateSP(this, &FMediaProfileEditor::SaveLayoutAs),
+		FCanExecuteAction::CreateSP(this, &FMediaProfileEditor::CanSaveLayoutAs));
+	
+	ToolkitCommands->MapAction(
+		FMediaProfileCommands::Get().RemoveAllLayouts,
+		FExecuteAction::CreateSP(this, &FMediaProfileEditor::RemoveAllLayouts));
+	
+	ToolkitCommands->MapAction(
+		FMediaProfileCommands::Get().ToggleFullscreen,
+		FExecuteAction::CreateSP(this, &FMediaProfileEditor::ToggleFullscreen),
+		FCanExecuteAction(),
+		FIsActionChecked::CreateSP(this, &FMediaProfileEditor::IsFullscreen));
+	
+	ToolkitCommands->MapAction(
+		FMediaProfileCommands::Get().Immersive,
+		FExecuteAction::CreateSP(this, &FMediaProfileEditor::SetViewportImmersive));
+}
+
 void FMediaProfileEditor::ExtendToolbar()
 {
 	TSharedPtr<FExtender> ToolbarExtender = MakeShared<FExtender>();
@@ -356,5 +623,39 @@
 		FToolBarExtensionDelegate::CreateLambda([this](FToolBarBuilder& ToolbarBuilder)
 		{
 			ToolbarBuilder.AddWidget(
+				SNew(SComboButton)
+				.OnGetMenuContent(this, &FMediaProfileEditor::CreateLayoutMenu)
+				.ButtonContent()
+				[
+					SNew(SHorizontalBox)
+					+ SHorizontalBox::Slot()
+					.AutoWidth()
+					.Padding(FMargin(0.f, 0.f, 6.f, 0.f))
+					[
+						SNew(SBox)
+						.WidthOverride(16.f)
+						.HeightOverride(16.f)
+						[
+							SNew(SImage).Image(FAppStyle::Get().GetBrush("Icons.Layout"))
+						]
+					]
+					+ SHorizontalBox::Slot()
+					[
+						SNew(STextBlock)
+						.Text(this, &FMediaProfileEditor::GetActiveLayoutText, FText::GetEmpty(), true)
+					]
+				]
+			);
+		})
+	);
+
+	ToolbarExtender->AddToolBarExtension(
+		"Asset",
+		EExtensionHook::After,
+		GetToolkitCommands(),
+		FToolBarExtensionDelegate::CreateLambda([this](FToolBarBuilder& ToolbarBuilder)
+		{
+			ToolbarBuilder.AddSeparator();
+			ToolbarBuilder.AddWidget(
 				SAssignNew(TimecodeToolbarEntry, SMediaFrameworkTimecodeGenlockHeader)
 				.Visibility_Lambda([]()
@@ -518,6 +819,334 @@
 		});
 }
 
+TSharedRef<SWidget> FMediaProfileEditor::CreateLayoutMenu()
+{
+	FMenuBuilder MenuBuilder(true, ToolkitCommands);
+					
+	MenuBuilder.BeginSection("CurrentLayout");
+	{
+		MenuBuilder.AddMenuEntry(
+			TAttribute<FText>::CreateSP(this, &FMediaProfileEditor::GetActiveLayoutText, LOCTEXT("CurrentLayoutLabel", "Current Layout"), false),
+			FText::GetEmpty(),
+			FSlateIcon(FAppStyle::Get().GetStyleSetName(), "Icons.Layout"),
+			FUIAction(FExecuteAction(), FCanExecuteAction::CreateLambda([]{ return false; })));
+	}
+	MenuBuilder.EndSection();
+	
+	MenuBuilder.BeginSection("SaveLayout", LOCTEXT("SaveLayoutSectionLabel", "Save Layout"));
+	{
+		MenuBuilder.AddMenuEntry(FMediaProfileCommands::Get().SaveLayout);
+		MenuBuilder.AddMenuEntry(FMediaProfileCommands::Get().SaveLayoutAs);
+	}
+	MenuBuilder.EndSection();
+	
+	MenuBuilder.BeginSection("LoadLayout", LOCTEXT("LoadLayoutSectionLabel", "Load Layout"));
+	{
+		MenuBuilder.AddMenuEntry(
+			LOCTEXT("DefaultLayout", "Default Layout"),
+			FText::GetEmpty(),
+			FSlateIcon(FMediaFrameworkUtilitiesEditorStyle::Get().GetStyleSetName(), "ToolbarIcon.LoadLayout"),
+			FUIAction(
+				FExecuteAction::CreateSP(this, &FMediaProfileEditor::LoadLayout, DefaultLayoutName),
+				FCanExecuteAction(),
+				FIsActionChecked::CreateSP(this, &FMediaProfileEditor::IsLayoutLoaded, DefaultLayoutName)),
+			NAME_None,
+			EUserInterfaceActionType::RadioButton);
+		
+		if (const UMediaProfileEditorUserSettings* UserSettings = GetDefault<UMediaProfileEditorUserSettings>())
+		{
+			for (const FMediaProfileEditorLayoutSettings& Layout : UserSettings->SavedLayouts)
+			{
+				MenuBuilder.AddMenuEntry(
+					FText::FromString(Layout.Name),
+					FText::GetEmpty(),
+					FSlateIcon(FMediaFrameworkUtilitiesEditorStyle::Get().GetStyleSetName(), "ToolbarIcon.LoadLayout"),
+					FUIAction(
+						FExecuteAction::CreateSP(this, &FMediaProfileEditor::LoadLayout, Layout.Name),
+						FCanExecuteAction(),
+						FIsActionChecked::CreateSP(this, &FMediaProfileEditor::IsLayoutLoaded, Layout.Name)),
+					NAME_None,
+					EUserInterfaceActionType::RadioButton);
+			}
+		}
+	}
+	MenuBuilder.EndSection();
+	
+	MenuBuilder.BeginSection("RemoveLayout", LOCTEXT("RemoveLayoutSectionLabel", "Remove Layout"));
+	{
+		MenuBuilder.AddSubMenu(
+			LOCTEXT("RemoveLayoutSubMenuLabel", "Remove Layout"),
+			FText::GetEmpty(),
+			FNewMenuDelegate::CreateLambda([this](FMenuBuilder& InSubMenuBuilder)
+			{
+				if (const UMediaProfileEditorUserSettings* UserSettings = GetDefault<UMediaProfileEditorUserSettings>())
+				{
+					for (const FMediaProfileEditorLayoutSettings& Layout : UserSettings->SavedLayouts)
+					{
+						InSubMenuBuilder.AddMenuEntry(
+							FText::FromString(Layout.Name),
+							FText::GetEmpty(),
+							FSlateIcon(),
+							FUIAction(FExecuteAction::CreateSP(this, &FMediaProfileEditor::RemoveLayout, Layout.Name)));
+					}
+				}
+				
+				InSubMenuBuilder.AddMenuSeparator();
+				InSubMenuBuilder.AddMenuEntry(FMediaProfileCommands::Get().RemoveAllLayouts);
+			}),
+			false,
+			FSlateIcon(FMediaFrameworkUtilitiesEditorStyle::Get().GetStyleSetName(), "ToolbarIcon.RemoveLayout"));
+	}
+	MenuBuilder.EndSection();
+	
+	MenuBuilder.BeginSection("Immersive");
+	{
+		MenuBuilder.AddMenuSeparator();
+		MenuBuilder.AddMenuEntry(FMediaProfileCommands::Get().ToggleFullscreen);
+	}
+	MenuBuilder.EndSection();
+	
+	return MenuBuilder.MakeWidget();
+}
+
+FText FMediaProfileEditor::GetActiveLayoutText(FText Prefix, bool bDirtyFlag) const
+{
+	if (const UMediaProfileEditorUserSettings* UserSettings = GetDefault<UMediaProfileEditorUserSettings>())
+	{
+		FString ActiveLayout = !UserSettings->ActiveLayout.IsEmpty() ? UserSettings->ActiveLayout : DefaultLayoutName;
+		
+		
+		FText ActiveLayoutLabel = FText::FromString(ActiveLayout);
+		if (bDirtyFlag && ActiveLayout != DefaultLayoutName && bActiveLayoutDirty)
+		{
+			ActiveLayoutLabel = bActiveLayoutDirty ? FText::Format(LOCTEXT("DirtyLayoutLabelFormat", "{0}*"), FText::FromString(ActiveLayout)) : FText::FromString(ActiveLayout);
+		}
+		
+		return !Prefix.IsEmpty()
+			? FText::Format(LOCTEXT("CurrentLayoutMenuEntryFormat", "{0}: {1}"), Prefix, ActiveLayoutLabel)
+			: ActiveLayoutLabel;
+	}
+	
+	return FText::GetEmpty();
+}
+
+void FMediaProfileEditor::SaveLayout()
+{
+	if (UMediaProfileEditorUserSettings* UserSettings = GetMutableDefault<UMediaProfileEditorUserSettings>())
+	{
+		const FString& ActiveLayout = UserSettings->ActiveLayout;
+		FMediaProfileEditorLayoutSettings* ExistingLayout = 
+			UserSettings->SavedLayouts.FindByPredicate([&ActiveLayout](const FMediaProfileEditorLayoutSettings& Layout) { return Layout.Name == ActiveLayout; });
+		
+		if (!ExistingLayout)
+		{
+			return;
+		}
+
+		if (!UE::MediaProfileEditor::Private::ExportLayoutToString(GetTabManager()->PersistLayout(), ActiveLayout, ExistingLayout->Layout))
+		{
+			return;
+		}
+		
+		if (ViewportPanel.IsValid())
+		{
+			ViewportPanel->SaveLayout(*ExistingLayout);
+		}
+		
+		UserSettings->SaveConfig();
+		
+		bActiveLayoutDirty = false;
+	}
+}
+
+bool FMediaProfileEditor::CanSaveLayout() const
+{
+	if (const UMediaProfileEditorUserSettings* UserSettings = GetDefault<UMediaProfileEditorUserSettings>())
+	{
+		return
+			!UserSettings->ActiveLayout.IsEmpty() && 
+			UserSettings->ActiveLayout != DefaultLayoutName && 
+			UserSettings->SavedLayouts.ContainsByPredicate([Name = UserSettings->ActiveLayout](const FMediaProfileEditorLayoutSettings& LayoutSettings)
+			{
+				return LayoutSettings.Name == Name;
+			});
+	}
+	
+	return false;
+}
+
+void FMediaProfileEditor::SaveLayoutAs()
+{
+	const FString NewLayoutName = UE::MediaProfileEditor::Private::ShowSaveLayoutAsDialog();
+	if (NewLayoutName.IsEmpty())
+	{
+		return;
+	}
+	
+	if (UMediaProfileEditorUserSettings* UserSettings = GetMutableDefault<UMediaProfileEditorUserSettings>())
+	{
+		FMediaProfileEditorLayoutSettings NewLayout;
+		NewLayout.Name = NewLayoutName;
+		
+		if (!UE::MediaProfileEditor::Private::ExportLayoutToString(GetTabManager()->PersistLayout(), NewLayoutName, NewLayout.Layout))
+		{
+			return;
+		}
+		
+		if (ViewportPanel.IsValid())
+		{
+			ViewportPanel->SaveLayout(NewLayout);
+		}
+		
+		UserSettings->SavedLayouts.Add(NewLayout);
+		UserSettings->SaveConfig();
+		
+		LoadLayout(NewLayoutName);
+	}
+}
+
+bool FMediaProfileEditor::CanSaveLayoutAs() const
+{
+	if (const UMediaProfileEditorUserSettings* UserSettings = GetDefault<UMediaProfileEditorUserSettings>())
+	{
+		return
+			UserSettings->ActiveLayout.IsEmpty() ||
+			UserSettings->ActiveLayout == DefaultLayoutName;
+	}
+	
+	return false;
+}
+
+void FMediaProfileEditor::LoadLayout(FString LayoutName)
+{
+	if (UMediaProfileEditorUserSettings* UserSettings = GetMutableDefault<UMediaProfileEditorUserSettings>())
+	{
+		if (LayoutName.IsEmpty() || LayoutName == DefaultLayoutName)
+		{
+			// Don't load any saved modifications to the default layout to ensure the layout gets set to true default
+			const bool bLoadUserLayout = false;
+			RestoreFromLayout(GetDefaultLayout(), bLoadUserLayout);
+			RegenerateMenusAndToolbars();
+		
+			if (ViewportPanel.IsValid())
+        	{
+				ViewportPanel->ClearLayout();
+			}
+		}
+		else
+		{
+			const FMediaProfileEditorLayoutSettings* ExistingLayout = 
+				UserSettings->SavedLayouts.FindByPredicate([&LayoutName](const FMediaProfileEditorLayoutSettings& Layout) { return Layout.Name == LayoutName; });
+			if (!ExistingLayout)
+			{
+				return;
+			}
+	
+			if (TSharedPtr<FTabManager::FLayout> NewLayout = UE::MediaProfileEditor::Private::ImportLayoutFromString(ExistingLayout->Layout))
+			{
+				const bool bLoadUserLayout = false;
+				RestoreFromLayout(NewLayout.ToSharedRef(), bLoadUserLayout);
+				RegenerateMenusAndToolbars();
+			
+				if (ViewportPanel.IsValid())
+				{
+					ViewportPanel->LoadLayout(*ExistingLayout);
+				}
+			}
+		}
+		
+		UserSettings->ActiveLayout = LayoutName;
+		UserSettings->SaveConfig();
+		
+		bActiveLayoutDirty = false;
+	}
+}
+
+bool FMediaProfileEditor::IsLayoutLoaded(FString LayoutName) const
+{
+	if (const UMediaProfileEditorUserSettings* UserSettings = GetDefault<UMediaProfileEditorUserSettings>())
+	{
+		const bool bIsDefaultLayout = LayoutName.IsEmpty() || LayoutName == DefaultLayoutName;
+		return bIsDefaultLayout
+			? UserSettings->ActiveLayout.IsEmpty() || UserSettings->ActiveLayout == DefaultLayoutName
+			: UserSettings->ActiveLayout == LayoutName;
+	}
+	
+	return false;
+}
+
+void FMediaProfileEditor::RemoveLayout(FString LayoutName)
+{
+	if (UMediaProfileEditorUserSettings* UserSettings = GetMutableDefault<UMediaProfileEditorUserSettings>())
+	{
+		const int32 ExistingLayoutIndex = 
+				UserSettings->SavedLayouts.IndexOfByPredicate([&LayoutName](const FMediaProfileEditorLayoutSettings& Layout) { return Layout.Name == LayoutName; });
+		if (ExistingLayoutIndex == INDEX_NONE)
+		{
+			return;
+		}
+		
+		if (UserSettings->ActiveLayout == UserSettings->SavedLayouts[ExistingLayoutIndex].Name)
+		{
+			LoadLayout(DefaultLayoutName);
+		}
+		
+		UserSettings->SavedLayouts.RemoveAt(ExistingLayoutIndex);
+		UserSettings->SaveConfig();
+	}
+}
+
+void FMediaProfileEditor::RemoveAllLayouts()
+{
+	if (UMediaProfileEditorUserSettings* UserSettings = GetMutableDefault<UMediaProfileEditorUserSettings>())
+	{
+		if (!UserSettings->ActiveLayout.IsEmpty() && UserSettings->ActiveLayout != DefaultLayoutName)
+		{
+			LoadLayout(DefaultLayoutName);
+		}
+		
+		UserSettings->SavedLayouts.Empty();
+		UserSettings->SaveConfig();
+	}
+}
+
+void FMediaProfileEditor::ToggleFullscreen()
+{
+	TSharedPtr<SDockTab> Tab = TabManager->GetOwnerTab();
+	const TSharedPtr<SWindow> Window = FSlateApplication::Get().FindWidgetWindow(Tab.ToSharedRef());
+
+	if (Window->GetWindowMode() == EWindowMode::Windowed)
+	{
+		Window->SetWindowMode(EWindowMode::WindowedFullscreen);
+	}
+	else
+	{
+		Window->SetWindowMode(EWindowMode::Windowed);
+	}
+}
+
+bool FMediaProfileEditor::IsFullscreen() const
+{
+	TSharedPtr<SDockTab> Tab = TabManager->GetOwnerTab();
+	const TSharedPtr<SWindow> Window = FSlateApplication::Get().FindWidgetWindow(Tab.ToSharedRef());
+
+	return Window->GetWindowMode() == EWindowMode::WindowedFullscreen;
+}
+
+void FMediaProfileEditor::SetViewportImmersive()
+{
+	if (ViewportPanel.IsValid())
+	{
+		ViewportPanel->SetImmersivePanel();
+	}
+}
+
+void FMediaProfileEditor::OnPersistLayout(const TSharedRef<FTabManager::FLayout>& LayoutToSave)
+{
+	bActiveLayoutDirty = true;
+	FLayoutSaveRestore::SaveToConfig(GEditorLayoutIni, LayoutToSave);
+}
+
 void FMediaProfileEditor::OnMediaItemDeleted(UClass* InMediaType, int32 InMediaItemIndex)
 {
 	if (InMediaType->IsChildOf<UMediaSource>())
diff --git a/Engine/Plugins/Media/MediaFrameworkUtilities/Source/MediaFrameworkUtilitiesEditor/Private/AssetEditor/MediaProfileEditor.h b/Engine/Plugins/Media/MediaFrameworkUtilities/Source/MediaFrameworkUtilitiesEditor/Private/AssetEditor/MediaProfileEditor.h
--- a/Engine/Plugins/Media/MediaFrameworkUtilities/Source/MediaFrameworkUtilitiesEditor/Private/AssetEditor/MediaProfileEditor.h
+++ b/Engine/Plugins/Media/MediaFrameworkUtilities/Source/MediaFrameworkUtilitiesEditor/Private/AssetEditor/MediaProfileEditor.h
@@ -28,6 +28,8 @@ private:
 	static const FName MediaTreeTabId;
 	static const FName DetailsTabId;
 	static const FName TimecodeTabId;
+	
+	static const FString DefaultLayoutName;
 
 public:
 	DECLARE_MULTICAST_DELEGATE_OneParam(FOnCaptureMethodChanged, UMediaOutput*)
@@ -75,6 +77,9 @@ public:
 	FOnCaptureMethodChanged& GetOnCaptureMethodChanged() { return OnCaptureMethodChanged; }
 	
 private:
+	/** Creates the default layout to use for the media profile editor */
+	TSharedRef<FTabManager::FLayout> GetDefaultLayout() const;
+	
 	TSharedRef<SDockTab> SpawnTab_MediaOutput(const FSpawnTabArgs& InArgs);
 	TSharedRef<SDockTab> SpawnTab_MediaTree(const FSpawnTabArgs& InArgs);
 	TSharedRef<SDockTab> SpawnTab_Details(const FSpawnTabArgs& InArgs);
@@ -82,6 +87,9 @@ private:
 
 	/** Raised when something has requested that a tab be manually opened/focused */
 	void OpenTab(FName InTabId);
+
+	/** Binds commands to the asset editor's command list */
+	void BindCommands();
 	
 	void ExtendToolbar();
 
@@ -91,6 +99,47 @@ private:
 	/** Cleans up and saves the media profile editor user settings config */
 	void CleanUpAndSaveUserSettings();
 	
+	TSharedRef<SWidget> CreateLayoutMenu();
+
+	/** Gets the text to display in the layouts menu for the currently active layout */
+	FText GetActiveLayoutText(FText Prefix, bool bDirtyFlag) const;
+
+	/** Raised when the 'Save Layout' menu entry is clicked in the Layouts dropdown */
+	void SaveLayout();
+
+	/** Gets whether the 'Save Layout' menu entry can be used on the current layout */
+	bool CanSaveLayout() const;
+
+	/** Raised when the 'Save Layout As' menu entry is clicked in the Layouts dropdown */
+	void SaveLayoutAs();
+
+	/** Gets whether the 'Save Layout As' menu entry can be used on the current layout */
+	bool CanSaveLayoutAs() const;
+
+	/** Raised when a specific layout is selected to load in the Layouts dropdown */
+	void LoadLayout(FString LayoutName);
+
+	/** Gets whether the layout of the specified name is the currently active layout */
+	bool IsLayoutLoaded(FString LayoutName) const;
+
+	/** Removes the layout of the specified name from the list of saved layouts */
+	void RemoveLayout(FString LayoutName);
+
+	/** Removes all saved layouts */
+	void RemoveAllLayouts();
+
+	/** Toggles fullscreen mode on the window that contains the media profile editor */
+	void ToggleFullscreen();
+
+	/** Gets whether the window containing the media profile editor is in fullscreen mode or not */
+	bool IsFullscreen() const;
+
+	/** Sets the media profile editor viewport to immersive mode */
+	void SetViewportImmersive();
+
+	/** Raised when the tab manager for this editor tries to save the current tab layout to config */
+	void OnPersistLayout(const TSharedRef<FTabManager::FLayout>& LayoutToSave);
+
 	void OnMediaItemDeleted(UClass* InMediaType, int32 InMediaItemIndex);
 	void OnSelectedMediaItemsChanged(const TArray<int32>& SelectedMediaSources, const TArray<int32>& SelectedMediaOutputs);
 	
@@ -113,6 +162,9 @@ private:
 	/** List of media outputs whose capture needs to be restarted on the next engine tick */
 	TSet<TWeakObjectPtr<UMediaOutput>> OutputsToRestartCapture;
 
+	/** Indicates if the currently active layout has been dirtied and can be saved */
+	bool bActiveLayoutDirty = false;
+	
 	/** Raised when the capture method is changed on a media output through the details panel */
 	FOnCaptureMethodChanged OnCaptureMethodChanged;
 	
diff --git a/Engine/Plugins/Media/MediaFrameworkUtilities/Source/MediaFrameworkUtilitiesEditor/Private/AssetEditor/MediaProfileEditorUserSettings.h b/Engine/Plugins/Media/MediaFrameworkUtilities/Source/MediaFrameworkUtilitiesEditor/Private/AssetEditor/MediaProfileEditorUserSettings.h
--- a/Engine/Plugins/Media/MediaFrameworkUtilities/Source/MediaFrameworkUtilitiesEditor/Private/AssetEditor/MediaProfileEditorUserSettings.h
+++ b/Engine/Plugins/Media/MediaFrameworkUtilities/Source/MediaFrameworkUtilitiesEditor/Private/AssetEditor/MediaProfileEditorUserSettings.h
@@ -49,6 +49,46 @@ public:
 	bool bFlipVertical = false;
 };
 
+USTRUCT()
+struct FMediaProfileEditorLayoutPanelSettings
+{
+	GENERATED_BODY()
+	
+	/** The name of the media item being displayed in the panel */
+	UPROPERTY(Config)
+	FString MediaItemName;
+	
+	/** Whether the media item being displayed is a media source or media output */
+	UPROPERTY(Config)
+	bool bMediaSource;
+};
+
+USTRUCT()
+struct FMediaProfileEditorLayoutSettings
+{
+	GENERATED_BODY()
+	
+	/** The name of the saved layout */
+	UPROPERTY(Config)
+	FString Name;
+	
+	/** The stringified FTabManager::FLayout of the saved layout */
+	UPROPERTY(Config)
+	FString Layout;
+	
+	/** A flag that indicates which panel arrangement the viewport has for the layout */
+	UPROPERTY(Config)
+	uint8 ViewportLayout;
+	
+	/** A flag that indicates which panel orientation the viewport has for the layout */
+	UPROPERTY(Config)
+	uint8 ViewportOrientation;
+	
+	/** A list of panel contents the viewport has for the layout */
+	UPROPERTY(Config)
+	TArray<FMediaProfileEditorLayoutPanelSettings> Panels; 
+};
+
 /** Config class to store editor user settings for the media profile editor */
 UCLASS(MinimalAPI, Config=EditorPerProjectUserSettings)
 class UMediaProfileEditorUserSettings : public UObject
@@ -79,6 +119,14 @@ public:
 	/** Config settings specific to individual media sources/outputs, indexed by the media item's name */
 	UPROPERTY(Config)
 	TMap<FName, FMediaProfileEditorPerMediaItemSettings> PerMediaItemSettings;
+	
+	/** A list of user-created layouts for the media profile editor */
+	UPROPERTY(Config)
+	TArray<FMediaProfileEditorLayoutSettings> SavedLayouts;
+	
+	/** The currently active layout for the media profile editor */
+	UPROPERTY(Config)
+	FString ActiveLayout;
 };
 
 /**
diff --git a/Engine/Plugins/Media/MediaFrameworkUtilities/Source/MediaFrameworkUtilitiesEditor/Private/AssetEditor/SMediaProfileViewport.cpp b/Engine/Plugins/Media/MediaFrameworkUtilities/Source/MediaFrameworkUtilitiesEditor/Private/AssetEditor/SMediaProfileViewport.cpp
--- a/Engine/Plugins/Media/MediaFrameworkUtilities/Source/MediaFrameworkUtilitiesEditor/Private/AssetEditor/SMediaProfileViewport.cpp
+++ b/Engine/Plugins/Media/MediaFrameworkUtilities/Source/MediaFrameworkUtilitiesEditor/Private/AssetEditor/SMediaProfileViewport.cpp
@@ -284,6 +284,107 @@ void SMediaProfileViewport::ClearImmersivePanel()
 	}
 }
 
+void SMediaProfileViewport::SaveLayout(FMediaProfileEditorLayoutSettings& InOutLayoutSettings)
+{
+	UMediaProfile* MediaProfile = nullptr;
+	if (MediaProfileEditor.IsValid())
+	{
+		MediaProfile = MediaProfileEditor.Pin()->GetMediaProfile();
+	}
+	
+	if (!MediaProfile)
+	{
+		return;
+	}
+	
+	InOutLayoutSettings.ViewportLayout = (uint8)CurrentLayout;
+	InOutLayoutSettings.ViewportOrientation = (uint8)CurrentOrientation;
+	
+	InOutLayoutSettings.Panels.Empty();
+	for (const FPanelContent& Panel : PanelContents)
+	{
+		FMediaProfileEditorLayoutPanelSettings PanelSettings;
+		if (Panel.MediaType->IsChildOf(UMediaSource::StaticClass()))
+		{
+			PanelSettings.MediaItemName = MediaProfile->GetLabelForMediaSource(Panel.MediaItemIndex);
+			PanelSettings.bMediaSource = true;
+		}
+		else
+		{
+			PanelSettings.MediaItemName = MediaProfile->GetLabelForMediaOutput(Panel.MediaItemIndex);
+			PanelSettings.bMediaSource = false;
+		}
+		
+		InOutLayoutSettings.Panels.Add(PanelSettings);
+	}
+}
+
+void SMediaProfileViewport::LoadLayout(const FMediaProfileEditorLayoutSettings& InLayoutSettings)
+{
+	UMediaProfile* MediaProfile = nullptr;
+	if (MediaProfileEditor.IsValid())
+	{
+		MediaProfile = MediaProfileEditor.Pin()->GetMediaProfile();
+	}
+	
+	if (!MediaProfile)
+	{
+		return;
+	}
+	
+	CurrentLayout = (EViewportLayout)InLayoutSettings.ViewportLayout;
+	CurrentOrientation = (EViewportOrientation)InLayoutSettings.ViewportOrientation;
+	
+	PanelContents.Empty();
+	for (const FMediaProfileEditorLayoutPanelSettings& PanelSettings : InLayoutSettings.Panels)
+	{
+		FPanelContent Panel;
+		
+		if (PanelSettings.bMediaSource)
+		{
+			Panel.MediaItemIndex = MediaProfile->FindMediaSourceIndexByLabel(PanelSettings.MediaItemName);
+
+			if (UMediaSource* MediaSource = MediaProfile->GetMediaSource(Panel.MediaItemIndex))
+			{
+				Panel.MediaType = MediaSource->GetClass();
+			}
+			else
+			{
+				Panel.MediaType = UMediaSource::StaticClass();
+			}
+		}
+		else
+		{
+			Panel.MediaItemIndex = MediaProfile->FindMediaOutputIndexByLabel(PanelSettings.MediaItemName);
+
+			if (UMediaOutput* MediaOutput = MediaProfile->GetMediaOutput(Panel.MediaItemIndex))
+			{
+				Panel.MediaType = MediaOutput->GetClass();
+			}
+			else
+			{
+				Panel.MediaType = UMediaOutput::StaticClass();
+			}
+		}
+
+		Panel.bIsLocked = true;
+		
+		PanelContents.Add(Panel);
+	}
+	
+	UpdateDisplays();
+}
+
+void SMediaProfileViewport::ClearLayout()
+{
+	CurrentLayout = EViewportLayout::OnePanel;
+	CurrentOrientation = EViewportOrientation::Left;
+	
+	PanelContents.Empty();
+
+	UpdateDisplays();
+}
+
 FReply SMediaProfileViewport::OnKeyDown(const FGeometry& MyGeometry, const FKeyEvent& InKeyEvent)
 {
 	FReply Reply = FReply::Unhandled();
diff --git a/Engine/Plugins/Media/MediaFrameworkUtilities/Source/MediaFrameworkUtilitiesEditor/Private/AssetEditor/SMediaProfileViewport.h b/Engine/Plugins/Media/MediaFrameworkUtilities/Source/MediaFrameworkUtilitiesEditor/Private/AssetEditor/SMediaProfileViewport.h
--- a/Engine/Plugins/Media/MediaFrameworkUtilities/Source/MediaFrameworkUtilitiesEditor/Private/AssetEditor/SMediaProfileViewport.h
+++ b/Engine/Plugins/Media/MediaFrameworkUtilities/Source/MediaFrameworkUtilitiesEditor/Private/AssetEditor/SMediaProfileViewport.h
@@ -5,6 +5,7 @@
 #include "ToolMenuEntry.h"
 #include "Widgets/SCompoundWidget.h"
 
+struct FMediaProfileEditorLayoutSettings;
 class FMediaProfileEditor;
 class SOverlay;
 class FUICommandList;
@@ -90,6 +91,15 @@ public:
 	/** Clears any immersive panel being displayed */
 	void ClearImmersivePanel();
 
+	/** Saves the current layout of the viewport's panels into the specified named layout */
+	void SaveLayout(FMediaProfileEditorLayoutSettings& InOutLayoutSettings);
+
+	/** Loads the panel configuration from the specified named layout into the viewport's panels */
+	void LoadLayout(const FMediaProfileEditorLayoutSettings& InLayoutSettings);
+
+	/** Clears the current panel layout back to default */
+	void ClearLayout();
+
 	/** Gets the viewport's command list */
 	TSharedPtr<FUICommandList> GetCommandList() { return CommandList; }
 	
diff --git a/Engine/Plugins/Media/MediaFrameworkUtilities/Source/MediaFrameworkUtilitiesEditor/Private/UI/MediaFrameworkUtilitiesEditorStyle.cpp b/Engine/Plugins/Media/MediaFrameworkUtilities/Source/MediaFrameworkUtilitiesEditor/Private/UI/MediaFrameworkUtilitiesEditorStyle.cpp
--- a/Engine/Plugins/Media/MediaFrameworkUtilities/Source/MediaFrameworkUtilitiesEditor/Private/UI/MediaFrameworkUtilitiesEditorStyle.cpp
+++ b/Engine/Plugins/Media/MediaFrameworkUtilities/Source/MediaFrameworkUtilitiesEditor/Private/UI/MediaFrameworkUtilitiesEditorStyle.cpp
@@ -28,7 +28,7 @@ namespace MediaFrameworkUtilitiesStyle
 void FMediaFrameworkUtilitiesEditorStyle::Register()
 {
 	MediaFrameworkUtilitiesStyle::StyleInstance = MakeUnique<FSlateStyleSet>(MediaFrameworkUtilitiesStyle::NAME_StyleName);
-	MediaFrameworkUtilitiesStyle::StyleInstance->SetCoreContentRoot(FPaths::EngineContentDir() / TEXT("Slate"));
+	MediaFrameworkUtilitiesStyle::StyleInstance->SetCoreContentRoot(FPaths::EngineContentDir());
 	MediaFrameworkUtilitiesStyle::StyleInstance->SetContentRoot(FPaths::EnginePluginsDir() / TEXT("Media/MediaFrameworkUtilities/Content/Editor/Icons/"));
 
 	MediaFrameworkUtilitiesStyle::StyleInstance->Set("ClassThumbnail.MediaBundle", new IMAGE_BRUSH("MediaBundle_64x", MediaFrameworkUtilitiesStyle::Icon64x64));
@@ -42,6 +42,11 @@ void FMediaFrameworkUtilitiesEditorStyle::Register()
 
 	MediaFrameworkUtilitiesStyle::StyleInstance->Set("ToolbarIcon.Edit", new IMAGE_BRUSH_SVG("MediaProfile", MediaFrameworkUtilitiesStyle::Icon40x40));
 	MediaFrameworkUtilitiesStyle::StyleInstance->Set("ToolbarIcon.Apply", new IMAGE_BRUSH("Apply_40x", MediaFrameworkUtilitiesStyle::Icon40x40));
+	MediaFrameworkUtilitiesStyle::StyleInstance->Set("ToolbarIcon.SaveLayout", new CORE_IMAGE_BRUSH_SVG("Editor/Slate/Starship/Common/LayoutSave", MediaFrameworkUtilitiesStyle::Icon16x16));
+	MediaFrameworkUtilitiesStyle::StyleInstance->Set("ToolbarIcon.SaveLayoutAs", new CORE_IMAGE_BRUSH_SVG( "Editor/Slate/Starship/Common/SaveCurrentAs", MediaFrameworkUtilitiesStyle::Icon16x16));
+	MediaFrameworkUtilitiesStyle::StyleInstance->Set("ToolbarIcon.LoadLayout", new CORE_IMAGE_BRUSH_SVG( "Editor/Slate/Starship/Common/LayoutLoad", MediaFrameworkUtilitiesStyle::Icon16x16));
+	MediaFrameworkUtilitiesStyle::StyleInstance->Set("ToolbarIcon.RemoveLayout", new CORE_IMAGE_BRUSH_SVG( "Editor/Slate/Starship/Common/LayoutRemove", MediaFrameworkUtilitiesStyle::Icon16x16));
+	MediaFrameworkUtilitiesStyle::StyleInstance->Set("ToolbarIcon.ToggleFullscreen", new CORE_IMAGE_BRUSH_SVG("Editor/Slate/Starship/Common/EnableFullscreen", MediaFrameworkUtilitiesStyle::Icon16x16));
 
 	MediaFrameworkUtilitiesStyle::StyleInstance->Set("TabIcons.MediaCapture.Small", new IMAGE_BRUSH_SVG("MediaCapture", MediaFrameworkUtilitiesStyle::Icon16x16));
 	MediaFrameworkUtilitiesStyle::StyleInstance->Set("MediaCapture.Capture", new IMAGE_BRUSH("PlayControlsRecord", MediaFrameworkUtilitiesStyle::Icon20x20));
@@ -56,7 +61,7 @@ void FMediaFrameworkUtilitiesEditorStyle::Register()
 	MediaFrameworkUtilitiesStyle::StyleInstance->Set("ToolbarIcon.Timecode", new IMAGE_BRUSH_SVG("Timecode", MediaFrameworkUtilitiesStyle::Icon16x16));
 	MediaFrameworkUtilitiesStyle::StyleInstance->Set("ToolbarIcon.Genlock", new IMAGE_BRUSH_SVG("Genlock", MediaFrameworkUtilitiesStyle::Icon16x16));
 
-	MediaFrameworkUtilitiesStyle::StyleInstance->Set("ToolbarIcon.Refresh", new CORE_IMAGE_BRUSH_SVG("Starship/Common/Update", MediaFrameworkUtilitiesStyle::Icon16x16));
+	MediaFrameworkUtilitiesStyle::StyleInstance->Set("ToolbarIcon.Refresh", new CORE_IMAGE_BRUSH_SVG("Slate/Starship/Common/Update", MediaFrameworkUtilitiesStyle::Icon16x16));
 	
 	//Use existing icon of MaterialInstanceConstant from Engine Editor content
 	const FString EngineContentPath = FPaths::EngineContentDir() / TEXT("Editor/Slate/Icons/AssetIcons/MaterialInstanceConstant_64x.png");
