Media Profile: Added toolbar button to media profile editor that can set or clear the media profile being edited as the active media profile. Also added Reset Transform option to the transform menu.

#Synced-CL 52395677

diff --git a/Engine/Plugins/Media/MediaFrameworkUtilities/Source/MediaFrameworkUtilitiesEditor/Private/AssetEditor/MediaProfileCommands.cpp b/Engine/Plugins/Media/MediaFrameworkUtilities/Source/MediaFrameworkUtilitiesEditor/Private/AssetEditor/MediaProfileCommands.cpp
--- a/Engine/Plugins/Media/MediaFrameworkUtilities/Source/MediaFrameworkUtilitiesEditor/Private/AssetEditor/MediaProfileCommands.cpp
+++ b/Engine/Plugins/Media/MediaFrameworkUtilities/Source/MediaFrameworkUtilitiesEditor/Private/AssetEditor/MediaProfileCommands.cpp
@@ -11,6 +11,7 @@ void FMediaProfileCommands::RegisterCommands()
 {
 	UI_COMMAND(Apply, "Apply", "Apply changes to the media profile.", EUserInterfaceActionType::Button, FInputChord());
 	UI_COMMAND(Edit, "Edit", "Edit the current media profile.", EUserInterfaceActionType::Button, FInputChord());
+	UI_COMMAND(ClearCurrentMediaProfile, "Clear Current Media Profile", "Removes this media profile as the current editor media profile", EUserInterfaceActionType::Button, FInputChord());
 	UI_COMMAND(SaveLayout, "Save Layout", "Save the current layout", EUserInterfaceActionType::Button, FInputChord());
 	UI_COMMAND(SaveLayoutAs, "Save Layout As", "Saves the current layout as...", EUserInterfaceActionType::Button, FInputChord());
 	UI_COMMAND(RemoveAllLayouts, "Remove All Layouts", "Removes all the user created layouts", EUserInterfaceActionType::Button, FInputChord());
diff --git a/Engine/Plugins/Media/MediaFrameworkUtilities/Source/MediaFrameworkUtilitiesEditor/Private/AssetEditor/MediaProfileCommands.h b/Engine/Plugins/Media/MediaFrameworkUtilities/Source/MediaFrameworkUtilitiesEditor/Private/AssetEditor/MediaProfileCommands.h
--- a/Engine/Plugins/Media/MediaFrameworkUtilities/Source/MediaFrameworkUtilitiesEditor/Private/AssetEditor/MediaProfileCommands.h
+++ b/Engine/Plugins/Media/MediaFrameworkUtilities/Source/MediaFrameworkUtilitiesEditor/Private/AssetEditor/MediaProfileCommands.h
@@ -28,6 +28,9 @@ public:
 	/** Edit the current Media Profile. */
 	TSharedPtr<FUICommandInfo> Edit;
 	
+	/** Clears out the current media profile  */
+	TSharedPtr<FUICommandInfo> ClearCurrentMediaProfile;
+	
 	/** Saves the current media profile editor layout */
 	TSharedPtr<FUICommandInfo> SaveLayout;
 	
diff --git a/Engine/Plugins/Media/MediaFrameworkUtilities/Source/MediaFrameworkUtilitiesEditor/Private/AssetEditor/MediaProfileEditor.cpp b/Engine/Plugins/Media/MediaFrameworkUtilities/Source/MediaFrameworkUtilitiesEditor/Private/AssetEditor/MediaProfileEditor.cpp
--- a/Engine/Plugins/Media/MediaFrameworkUtilities/Source/MediaFrameworkUtilitiesEditor/Private/AssetEditor/MediaProfileEditor.cpp
+++ b/Engine/Plugins/Media/MediaFrameworkUtilities/Source/MediaFrameworkUtilitiesEditor/Private/AssetEditor/MediaProfileEditor.cpp
@@ -587,6 +587,11 @@ void FMediaProfileEditor::OpenTab(FName InTabId)
 
 void FMediaProfileEditor::BindCommands()
 {
+	ToolkitCommands->MapAction(
+		FMediaProfileCommands::Get().ClearCurrentMediaProfile,
+		FExecuteAction::CreateSP(this, &FMediaProfileEditor::ClearCurrentMediaProfile),
+		FCanExecuteAction::CreateSP(this, &FMediaProfileEditor::CanClearCurrentMediaProfile));
+	
 	ToolkitCommands->MapAction(
 		FMediaProfileCommands::Get().SaveLayout, 
 		FExecuteAction::CreateSP(this, &FMediaProfileEditor::SaveLayout),
@@ -622,6 +627,30 @@ void FMediaProfileEditor::ExtendToolbar()
 		GetToolkitCommands(),
 		FToolBarExtensionDelegate::CreateLambda([this](FToolBarBuilder& ToolbarBuilder)
 		{
+			ToolbarBuilder.AddToolBarButton(
+				FUIAction(FExecuteAction::CreateSP(this, &FMediaProfileEditor::MakeCurrentMediaProfile)),
+				NAME_None,
+				TAttribute<FText>::CreateSP(this, &FMediaProfileEditor::GetIsCurrentMediaProfileLabel),
+				FText::GetEmpty(),
+				FSlateIcon(FMediaFrameworkUtilitiesEditorStyle::GetStyleSetName(), "ClassIcon.MediaProfile"));
+				
+			ToolbarBuilder.AddComboButton(
+				FUIAction(),
+				FOnGetContent::CreateSP(this, &FMediaProfileEditor::CreateCurrentProfileMenu),
+				TAttribute<FText>::CreateSP(this, &FMediaProfileEditor::GetIsCurrentMediaProfileLabel),
+				FText::GetEmpty(),
+				FSlateIcon(),
+				true);
+		})
+	);
+	
+	ToolbarExtender->AddToolBarExtension(
+		"Asset",
+		EExtensionHook::After,
+		GetToolkitCommands(),
+		FToolBarExtensionDelegate::CreateLambda([this](FToolBarBuilder& ToolbarBuilder)
+		{
+			ToolbarBuilder.AddSeparator();
 			ToolbarBuilder.AddWidget(
 				SNew(SComboButton)
 				.OnGetMenuContent(this, &FMediaProfileEditor::CreateLayoutMenu)
@@ -819,6 +848,41 @@ void FMediaProfileEditor::RestartActiveMediaCaptures(UMediaOutput* InMediaOutput
 		});
 }
 
+FText FMediaProfileEditor::GetIsCurrentMediaProfileLabel() const
+{
+	return IMediaProfileManager::Get().GetCurrentMediaProfile() == MediaProfileBeingEdited
+		? LOCTEXT("CurrentProfileLabel", "Current")
+		: LOCTEXT("MakeCurrentProfileLabel", "Make Current");
+}
+
+TSharedRef<SWidget> FMediaProfileEditor::CreateCurrentProfileMenu()
+{
+	FMenuBuilder MenuBuilder(true, ToolkitCommands);
+
+	MenuBuilder.BeginSection("CurrentMediaProfile", LOCTEXT("CurrentMediaProfileSectionLabel", "Current Media Profile"));
+	{
+		MenuBuilder.AddMenuEntry(FMediaProfileCommands::Get().ClearCurrentMediaProfile);
+	}
+	MenuBuilder.EndSection();
+
+	return MenuBuilder.MakeWidget();
+}
+
+void FMediaProfileEditor::MakeCurrentMediaProfile()
+{
+	IMediaProfileManager::Get().SetCurrentMediaProfile(MediaProfileBeingEdited);
+}
+
+void FMediaProfileEditor::ClearCurrentMediaProfile()
+{
+	IMediaProfileManager::Get().SetCurrentMediaProfile(nullptr);
+}
+
+bool FMediaProfileEditor::CanClearCurrentMediaProfile() const
+{
+	return IMediaProfileManager::Get().GetCurrentMediaProfile() == MediaProfileBeingEdited;
+}
+
 TSharedRef<SWidget> FMediaProfileEditor::CreateLayoutMenu()
 {
 	FMenuBuilder MenuBuilder(true, ToolkitCommands);
diff --git a/Engine/Plugins/Media/MediaFrameworkUtilities/Source/MediaFrameworkUtilitiesEditor/Private/AssetEditor/MediaProfileEditor.h b/Engine/Plugins/Media/MediaFrameworkUtilities/Source/MediaFrameworkUtilitiesEditor/Private/AssetEditor/MediaProfileEditor.h
--- a/Engine/Plugins/Media/MediaFrameworkUtilities/Source/MediaFrameworkUtilitiesEditor/Private/AssetEditor/MediaProfileEditor.h
+++ b/Engine/Plugins/Media/MediaFrameworkUtilities/Source/MediaFrameworkUtilitiesEditor/Private/AssetEditor/MediaProfileEditor.h
@@ -99,6 +99,21 @@ private:
 	/** Cleans up and saves the media profile editor user settings config */
 	void CleanUpAndSaveUserSettings();
 	
+	/** Gets the label used to indicate if the media profile being edited is set as the current editor media profile */
+	FText GetIsCurrentMediaProfileLabel() const;
+
+	/** Creates the dropdown menu for the "Current" profile toolbar button */
+	TSharedRef<SWidget> CreateCurrentProfileMenu();
+
+	/** Makes the media profile being edited the currently active media profile */
+	void MakeCurrentMediaProfile();
+	
+	/** Clears the media profile being edited from being the currently active media profile */
+	void ClearCurrentMediaProfile();
+
+	/** Gets whether the media profile being edited can be cleared from being the currently active media profile */
+	bool CanClearCurrentMediaProfile() const;
+
 	TSharedRef<SWidget> CreateLayoutMenu();
 
 	/** Gets the text to display in the layouts menu for the currently active layout */
diff --git a/Engine/Plugins/Media/MediaFrameworkUtilities/Source/MediaFrameworkUtilitiesEditor/Private/AssetEditor/SMediaProfileMediaItemDisplay.cpp b/Engine/Plugins/Media/MediaFrameworkUtilities/Source/MediaFrameworkUtilitiesEditor/Private/AssetEditor/SMediaProfileMediaItemDisplay.cpp
--- a/Engine/Plugins/Media/MediaFrameworkUtilities/Source/MediaFrameworkUtilitiesEditor/Private/AssetEditor/SMediaProfileMediaItemDisplay.cpp
+++ b/Engine/Plugins/Media/MediaFrameworkUtilities/Source/MediaFrameworkUtilitiesEditor/Private/AssetEditor/SMediaProfileMediaItemDisplay.cpp
@@ -1371,6 +1371,48 @@ FVector2f SMediaProfileMediaOutputDisplay::GetScale() const
 	return SMediaProfileMediaItemDisplayBase<UMediaOutput>::GetScale();
 }
 
+void SMediaProfileMediaOutputDisplay::ResetTransform()
+{
+	if (UMediaProfileEditorCaptureSettings* CaptureSettings = FMediaProfileEditor::GetMediaFrameworkCaptureSettings())
+	{
+		if (CaptureType == ECaptureType::MediaTexture &&
+			CaptureSettings->MediaTextureCaptures.IsValidIndex(CaptureIndex) &&
+			CaptureSettings->MediaTextureCaptures[CaptureIndex].MediaOutput == GetMediaItem())
+		{
+			CaptureSettings->MediaTextureCaptures[CaptureIndex].Transform.Rotation = EMediaCaptureRotation::None;
+			CaptureSettings->MediaTextureCaptures[CaptureIndex].Transform.bFlipHorizontal = false;
+			CaptureSettings->MediaTextureCaptures[CaptureIndex].Transform.bFlipVertical = false;
+			
+			if (TSharedPtr<FMediaProfileEditor> PinnedMediaProfileEditor = MediaProfileEditor.Pin())
+			{
+				PinnedMediaProfileEditor->RestartActiveMediaCaptures(GetMediaItem());
+			}
+			
+			return;
+		}
+	}
+	
+	SMediaProfileMediaItemDisplayBase<UMediaOutput>::ResetTransform();
+}
+
+bool SMediaProfileMediaOutputDisplay::CanResetTransform() const
+{
+	if (UMediaProfileEditorCaptureSettings* CaptureSettings = FMediaProfileEditor::GetMediaFrameworkCaptureSettings())
+	{
+		if (CaptureType == ECaptureType::MediaTexture &&
+			CaptureSettings->MediaTextureCaptures.IsValidIndex(CaptureIndex) &&
+			CaptureSettings->MediaTextureCaptures[CaptureIndex].MediaOutput == GetMediaItem())
+		{
+			return
+				CaptureSettings->MediaTextureCaptures[CaptureIndex].Transform.Rotation != EMediaCaptureRotation::None ||
+				CaptureSettings->MediaTextureCaptures[CaptureIndex].Transform.bFlipHorizontal ||
+				CaptureSettings->MediaTextureCaptures[CaptureIndex].Transform.bFlipVertical;
+		}
+	}
+	
+	return false;
+}
+
 void SMediaProfileMediaOutputDisplay::SetChannelMask(EColorChannelMask InChannelMask)
 {
 	SMediaProfileMediaItemDisplayBase::SetChannelMask(InChannelMask);
diff --git a/Engine/Plugins/Media/MediaFrameworkUtilities/Source/MediaFrameworkUtilitiesEditor/Private/AssetEditor/SMediaProfileMediaItemDisplay.h b/Engine/Plugins/Media/MediaFrameworkUtilities/Source/MediaFrameworkUtilitiesEditor/Private/AssetEditor/SMediaProfileMediaItemDisplay.h
--- a/Engine/Plugins/Media/MediaFrameworkUtilities/Source/MediaFrameworkUtilitiesEditor/Private/AssetEditor/SMediaProfileMediaItemDisplay.h
+++ b/Engine/Plugins/Media/MediaFrameworkUtilities/Source/MediaFrameworkUtilitiesEditor/Private/AssetEditor/SMediaProfileMediaItemDisplay.h
@@ -431,6 +431,14 @@ protected:
 		FToolMenuSection& TransformSection = ToolMenu->AddSection(TEXT("TransformSection"), LOCTEXT("TransformSectionLabel", "Transform"));
 		TransformSection.AddEntry(CreateRotationSubmenu());
 		TransformSection.AddEntry(CreateFlipSubmenu());
+		TransformSection.AddMenuEntry(
+			"ResetTransform",
+			LOCTEXT("ResetTransformLabel", "Reset Transform"),
+			LOCTEXT("ResetTransformToolTip", "Resets all transforms on the media item"),
+			 FSlateIcon(FAppStyle::GetAppStyleSetName(), "GenericCommands.Undo"),
+			 FUIAction(
+			 	FExecuteAction::CreateSP(this, &SMediaProfileMediaItemDisplayBase::ResetTransform),
+			 	FCanExecuteAction::CreateSP(this, &SMediaProfileMediaItemDisplayBase::CanResetTransform)));
 		
 		FToolMenuSection& ChannelsSection = ToolMenu->AddSection(TEXT("ChannelsSection"), LOCTEXT("ChannelsSectionLabel", "Channels"));
 		ChannelsSection.AddEntry(CreateChannelsSubmenu(true));
@@ -473,7 +481,8 @@ protected:
 				CreateRotationMenuEntry(EMediaImageRotation::CW180);
 			}),
 			/*bInOpenSubMenuOnClick=*/false,
-			FSlateIcon(FAppStyle::GetAppStyleSetName(), "EditorViewport.RotateMode")
+			FSlateIcon(FAppStyle::GetAppStyleSetName(), "EditorViewport.RotateMode"),
+			/*bInShouldCloseWindowAfterMenuSelection=*/false
 		);
 
 		return Entry;
@@ -539,7 +548,8 @@ protected:
 				);
 			}),
 			/*bInOpenSubMenuOnClick=*/false,
-			FSlateIcon(FAppStyle::GetAppStyleSetName(), "Icons.FlipVertical")
+			FSlateIcon(FAppStyle::GetAppStyleSetName(), "Icons.FlipVertical"),
+			/*bInShouldCloseWindowAfterMenuSelection=*/false
 		);
 
 		return Entry;
@@ -1043,6 +1053,30 @@ protected:
 		return Scale;
 	}
 	
+	/** Resets any active transform on the media item */
+	virtual void ResetTransform()
+	{
+		if (FMediaProfileEditorPerMediaItemSettings* MediaItemSettings = GetMediaItemSettings())
+		{
+			MediaItemSettings->Rotation = EMediaImageRotation::None;
+			MediaItemSettings->bFlipHorizontal = false;
+			MediaItemSettings->bFlipVertical = false;
+		}
+	}
+
+	/** Gets whether the media item has a transformation that can be reset */
+	virtual bool CanResetTransform() const
+	{
+		if (FMediaProfileEditorPerMediaItemSettings* MediaItemSettings = GetMediaItemSettings())
+		{
+			return MediaItemSettings->Rotation != EMediaImageRotation::None ||
+				MediaItemSettings->bFlipHorizontal ||
+				MediaItemSettings->bFlipVertical;
+		}
+		
+		return false;
+	}
+
 	/** Gets the current color channel being displayed in the viewport */
 	virtual EColorChannelMask GetChannelMask() const
 	{
@@ -1320,6 +1354,8 @@ protected:
 	virtual void ToggleFlipHorizontal() override;
 	virtual void ToggleFlipVertical() override;
 	virtual FVector2f GetScale() const override;
+	virtual void ResetTransform() override;
+	virtual bool CanResetTransform() const override;
 
 	virtual void SetChannelMask(EColorChannelMask InChannelMask) override;
 	virtual void SetInvertAlphaChannelMask(bool bInInvertAlphaChannelMask) override;
diff --git a/Engine/Plugins/Media/MediaFrameworkUtilities/Source/MediaFrameworkUtilitiesEditor/Private/UI/MediaFrameworkUtilitiesEditorStyle.cpp b/Engine/Plugins/Media/MediaFrameworkUtilities/Source/MediaFrameworkUtilitiesEditor/Private/UI/MediaFrameworkUtilitiesEditorStyle.cpp
--- a/Engine/Plugins/Media/MediaFrameworkUtilities/Source/MediaFrameworkUtilitiesEditor/Private/UI/MediaFrameworkUtilitiesEditorStyle.cpp
+++ b/Engine/Plugins/Media/MediaFrameworkUtilities/Source/MediaFrameworkUtilitiesEditor/Private/UI/MediaFrameworkUtilitiesEditorStyle.cpp
@@ -42,6 +42,7 @@ void FMediaFrameworkUtilitiesEditorStyle::Register()
 
 	MediaFrameworkUtilitiesStyle::StyleInstance->Set("ToolbarIcon.Edit", new IMAGE_BRUSH_SVG("MediaProfile", MediaFrameworkUtilitiesStyle::Icon40x40));
 	MediaFrameworkUtilitiesStyle::StyleInstance->Set("ToolbarIcon.Apply", new IMAGE_BRUSH("Apply_40x", MediaFrameworkUtilitiesStyle::Icon40x40));
+	MediaFrameworkUtilitiesStyle::StyleInstance->Set("ToolbarIcon.ClearCurrentMediaProfile", new IMAGE_BRUSH_SVG("MediaProfile", MediaFrameworkUtilitiesStyle::Icon16x16));
 	MediaFrameworkUtilitiesStyle::StyleInstance->Set("ToolbarIcon.SaveLayout", new CORE_IMAGE_BRUSH_SVG("Editor/Slate/Starship/Common/LayoutSave", MediaFrameworkUtilitiesStyle::Icon16x16));
 	MediaFrameworkUtilitiesStyle::StyleInstance->Set("ToolbarIcon.SaveLayoutAs", new CORE_IMAGE_BRUSH_SVG( "Editor/Slate/Starship/Common/SaveCurrentAs", MediaFrameworkUtilitiesStyle::Icon16x16));
 	MediaFrameworkUtilitiesStyle::StyleInstance->Set("ToolbarIcon.LoadLayout", new CORE_IMAGE_BRUSH_SVG( "Editor/Slate/Starship/Common/LayoutLoad", MediaFrameworkUtilitiesStyle::Icon16x16));
