MediaProfile: Extend SMediaProfileSourceTexturePicker to support configurable AllowedClass and asset filter.

We also support a custom widget extension.

#Synced-CL 52419094

diff --git a/Engine/Plugins/Media/MediaFrameworkUtilities/Source/MediaFrameworkUtilitiesEditor/Private/SMediaProfileSourceTexturePicker.cpp b/Engine/Plugins/Media/MediaFrameworkUtilities/Source/MediaFrameworkUtilitiesEditor/Private/SMediaProfileSourceTexturePicker.cpp
--- a/Engine/Plugins/Media/MediaFrameworkUtilities/Source/MediaFrameworkUtilitiesEditor/Private/SMediaProfileSourceTexturePicker.cpp
+++ b/Engine/Plugins/Media/MediaFrameworkUtilities/Source/MediaFrameworkUtilitiesEditor/Private/SMediaProfileSourceTexturePicker.cpp
@@ -24,45 +24,59 @@ void SMediaProfileSourceTexturePicker::Construct(const FArguments& InArgs)
 	TexturePropertyHandle = InArgs._TexturePropertyHandle;
 	OnMediaSourceSelected = InArgs._OnMediaSourceSelected;
 	
-	SObjectPropertyEntryBox::FArguments BaseArgs;
-	BaseArgs
-		.PropertyHandle(TexturePropertyHandle)
-		.AllowedClass(UMediaTexture::StaticClass())
-		.AllowCreate(true)
-		.ThumbnailPool(InArgs._ThumbnailPool)
-		.OnShouldFilterAsset_Lambda([](FAssetData const& AssetData)
+	UClass* EffectiveAllowedClass = InArgs._AllowedClass ? InArgs._AllowedClass : UMediaTexture::StaticClass();
+	const FOnShouldFilterAsset FilterDelegate = InArgs._OnShouldFilterAsset.IsBound()
+		? InArgs._OnShouldFilterAsset
+		: FOnShouldFilterAsset::CreateLambda([EffectiveAllowedClass](FAssetData const& AssetData)
 		{
 			if (UClass* AssetClass = AssetData.GetClass())
 			{
-				return !AssetClass->IsChildOf(UMediaTexture::StaticClass());
+				return !AssetClass->IsChildOf(EffectiveAllowedClass);
 			}
-
 			return true;
-		})
-		.CustomContentSlot()
-		[
-			SNew(SHorizontalBox)
+		});
 
-			+SHorizontalBox::Slot()
-			.AutoWidth()
+	TSharedRef<SHorizontalBox> ButtonsBox = SNew(SHorizontalBox)
+		+ SHorizontalBox::Slot()
+		.AutoWidth()
+		[
+			SNew(SComboButton)
+			.HasDownArrow(true)
+			.ComboButtonStyle(FAppStyle::Get(), "SimpleComboButton")
+			.ToolTipText(LOCTEXT("MediaProfileSourceSelectorToolTip", "Select media texture from a media source in the active Media Profile"))
+			.HAlign(HAlign_Center)
+			.VAlign(VAlign_Center)
+			.OnGetMenuContent(this, &SMediaProfileSourceTexturePicker::GetMediaProfileSourceSelectorMenu)
+			.IsEnabled(this, &SMediaProfileSourceTexturePicker::HasActiveMediaProfile)
+			.ButtonContent()
 			[
-				SNew(SComboButton)
-				.HasDownArrow(true)
-				.ComboButtonStyle(FAppStyle::Get(), "SimpleComboButton")
-				.ToolTipText(LOCTEXT("MediaProfileSourceSelectorToolTip", "Select media texture from a media source in the active Media Profile"))
-				.HAlign(HAlign_Center)
-				.VAlign(VAlign_Center)
-				.OnGetMenuContent(this, &SMediaProfileSourceTexturePicker::GetMediaProfileSourceSelectorMenu)
-				.IsEnabled(this, &SMediaProfileSourceTexturePicker::HasActiveMediaProfile)
-				.ButtonContent()
-				[
-					SNew(SImage)
-					.DesiredSizeOverride(FVector2D(16.0f, 16.0f))
-					.Image(FMediaFrameworkUtilitiesEditorStyle::Get().GetBrush("ClassIcon.MediaProfile"))
-				]
+				SNew(SImage)
+				.DesiredSizeOverride(FVector2D(16.0f, 16.0f))
+				.Image(FMediaFrameworkUtilitiesEditorStyle::Get().GetBrush("ClassIcon.MediaProfile"))
 			]
 		];
-	
+
+	if (TSharedPtr<SWidget> AdditionalWidget = InArgs._AdditionalContent.Widget)
+	{
+		ButtonsBox->AddSlot()
+		.AutoWidth()
+		[
+			AdditionalWidget.ToSharedRef()
+		];
+	}
+
+	SObjectPropertyEntryBox::FArguments BaseArgs;
+	BaseArgs
+		.PropertyHandle(TexturePropertyHandle)
+		.AllowedClass(EffectiveAllowedClass)
+		.AllowCreate(true)
+		.ThumbnailPool(InArgs._ThumbnailPool)
+		.OnShouldFilterAsset(FilterDelegate)
+		.CustomContentSlot()
+		[
+			ButtonsBox
+		];
+
 	SObjectPropertyEntryBox::Construct(BaseArgs);
 }
 
diff --git a/Engine/Plugins/Media/MediaFrameworkUtilities/Source/MediaFrameworkUtilitiesEditor/Public/SMediaProfileSourceTexturePicker.h b/Engine/Plugins/Media/MediaFrameworkUtilities/Source/MediaFrameworkUtilitiesEditor/Public/SMediaProfileSourceTexturePicker.h
--- a/Engine/Plugins/Media/MediaFrameworkUtilities/Source/MediaFrameworkUtilitiesEditor/Public/SMediaProfileSourceTexturePicker.h
+++ b/Engine/Plugins/Media/MediaFrameworkUtilities/Source/MediaFrameworkUtilitiesEditor/Public/SMediaProfileSourceTexturePicker.h
@@ -15,10 +15,18 @@ public:
 	DECLARE_DELEGATE_TwoParams(FOnMediaSourceSelected, UMediaProfile* /* InMediaProfile */, int32 /* InMediaSourceIndex */);
 	
 public:
-	SLATE_BEGIN_ARGS(SMediaProfileSourceTexturePicker) { }
+	SLATE_BEGIN_ARGS(SMediaProfileSourceTexturePicker)
+		: _AllowedClass(nullptr)
+		{ }
 		SLATE_ARGUMENT(TSharedPtr<IPropertyHandle>, TexturePropertyHandle)
 		SLATE_ARGUMENT(TSharedPtr<FAssetThumbnailPool>, ThumbnailPool)
+		/** Overrides the default UMediaTexture allowed class for the asset picker. */
+		SLATE_ARGUMENT(UClass*, AllowedClass)
+		/** Overrides the default MediaTexture-only asset filter. When unbound, only UMediaTexture assets are shown. */
+		SLATE_EVENT(FOnShouldFilterAsset, OnShouldFilterAsset)
 		SLATE_EVENT(FOnMediaSourceSelected, OnMediaSourceSelected)
+		/** Optional extra widget appended after the media profile button in the custom content slot. */
+		SLATE_NAMED_SLOT(FArguments, AdditionalContent)
 	SLATE_END_ARGS()
 
 	/** Constructs this widget with InArgs */
