MediaProfile: Fixed "not initialized properly" error logs.

#Synced-CL 52595449

diff --git a/Engine/Plugins/Media/MediaFrameworkUtilities/Source/MediaFrameworkUtilitiesEditor/Private/AssetEditor/MediaProfileEditorUserSettings.h b/Engine/Plugins/Media/MediaFrameworkUtilities/Source/MediaFrameworkUtilitiesEditor/Private/AssetEditor/MediaProfileEditorUserSettings.h
--- a/Engine/Plugins/Media/MediaFrameworkUtilities/Source/MediaFrameworkUtilitiesEditor/Private/AssetEditor/MediaProfileEditorUserSettings.h
+++ b/Engine/Plugins/Media/MediaFrameworkUtilities/Source/MediaFrameworkUtilitiesEditor/Private/AssetEditor/MediaProfileEditorUserSettings.h
@@ -60,7 +60,7 @@ struct FMediaProfileEditorLayoutPanelSettings
 	
 	/** Whether the media item being displayed is a media source or media output */
 	UPROPERTY(Config)
-	bool bMediaSource;
+	bool bMediaSource = false;
 };
 
 USTRUCT()
@@ -78,11 +78,11 @@ struct FMediaProfileEditorLayoutSettings
 	
 	/** A flag that indicates which panel arrangement the viewport has for the layout */
 	UPROPERTY(Config)
-	uint8 ViewportLayout;
-	
+	uint8 ViewportLayout = 0;
+
 	/** A flag that indicates which panel orientation the viewport has for the layout */
 	UPROPERTY(Config)
-	uint8 ViewportOrientation;
+	uint8 ViewportOrientation = 0;
 	
 	/** A list of panel contents the viewport has for the layout */
 	UPROPERTY(Config)
